UnrealTraceServer source.

diff --git a/Engine/Source/Programs/UnrealTraceServer/.gitignore b/Engine/Source/Programs/UnrealTraceServer/.gitignore
new file mode 100644
--- /dev/null
+++ b/Engine/Source/Programs/UnrealTraceServer/.gitignore
@@ -0,0 +1 @@
+_build
diff --git a/Engine/Source/Programs/UnrealTraceServer/build.bat b/Engine/Source/Programs/UnrealTraceServer/build.bat
new file mode 100644
--- /dev/null
+++ b/Engine/Source/Programs/UnrealTraceServer/build.bat
@@ -0,0 +1,2 @@
+@echo off
+nmake -fmakefile_nmake ASIO_PATH=d:\depot\starship-main\Engine\Source\ThirdParty\asio\1.12.2 %*
diff --git a/Engine/Source/Programs/UnrealTraceServer/makefile b/Engine/Source/Programs/UnrealTraceServer/makefile
new file mode 100644
--- /dev/null
+++ b/Engine/Source/Programs/UnrealTraceServer/makefile
@@ -0,0 +1,80 @@
+# vim: set foldlevel=1 :
+
+# {{{1 debug and flags ---------------------------------------------------------
+OUT_DIR := _build
+ifeq ($(DEBUG), 1)
+CFLAGS := \
+	$(CFLAGS) \
+	-DTS_BUILD_DEBUG=TS_ON \
+	-Og
+OUT_DIR := $(OUT_DIR)/debug
+else
+CFLAGS := \
+	$(CFLAGS) \
+	-O3
+OUT_DIR := $(OUT_DIR)/development
+endif
+
+CFLAGS := \
+	$(CFLAGS) \
+	-g \
+	-Werror \
+	-std=c++17 \
+	-fno-exceptions \
+	-static
+
+UNAME := $(shell uname)
+ifeq ($(UNAME), Linux)
+	LDLIBS := $(LDLIBS) -lstdc++ -lpthread -lrt
+endif
+
+ifeq ($(UNAME), Darwin)
+	LDLIBS := $(LDLIBS) -lstdc++ -lpthread
+endif
+
+# {{{1 all/clean targets -------------------------------------------------------
+.SILENT :
+
+.PHONY : clean all unreal_trace
+
+all : unreal_trace
+
+clean :
+	rm -rf $(OUT_DIR)
+
+$(OUT_DIR) :
+	mkdir -p $@
+
+# {{{1 asio --------------------------------------------------------------------
+ASIO_VER = 1.12.2
+CFLAGS := \
+	$(CFLAGS) \
+	-I3rdparty/asio/$(ASIO_VER)
+
+# {{{1 store compile/link ------------------------------------------------------
+OBJS := \
+	$(OBJS) \
+	$(OUT_DIR)/Asio.o \
+	$(OUT_DIR)/AsioContext.o \
+	$(OUT_DIR)/AsioFile.o \
+	$(OUT_DIR)/AsioIoable.o \
+	$(OUT_DIR)/AsioSocket.o \
+	$(OUT_DIR)/AsioTcpServer.o \
+	$(OUT_DIR)/AsioTickable.o \
+	$(OUT_DIR)/Cbor.o \
+	$(OUT_DIR)/Main.o \
+	$(OUT_DIR)/Recorder.o \
+	$(OUT_DIR)/Store.o \
+	$(OUT_DIR)/StoreCborServer.o \
+	$(OUT_DIR)/StoreService.o \
+	$(OUT_DIR)/TraceRelay.o
+
+$(OUT_DIR)/%.o : src/%.cpp
+	echo $(<F)
+	$(CXX) $(CFLAGS) -c -o $@ $<
+
+$(OUT_DIR)/UnrealTraceServer : $(OUT_DIR) $(OBJS)
+	echo $(@F)
+	$(CXX) -o $@ $(OBJS) $(LDLIBS)
+
+unreal_trace : $(OUT_DIR)/UnrealTraceServer
diff --git a/Engine/Source/Programs/UnrealTraceServer/src/Asio.cpp b/Engine/Source/Programs/UnrealTraceServer/src/Asio.cpp
new file mode 100644
--- /dev/null
+++ b/Engine/Source/Programs/UnrealTraceServer/src/Asio.cpp
@@ -0,0 +1,16 @@
+// Copyright Epic Games, Inc. All Rights Reserved.
+
+#include "Pch.h"
+#include "Asio.h"
+
+#if defined(_MSC_VER)
+#	include <malloc.h>
+#else
+#	include <alloca.h>
+#endif
+
+TS_ASIO_INCLUDE_BEGIN
+#include "asio/impl/src.hpp"
+TS_ASIO_INCLUDE_END
+
+/* vim: set noexpandtab : */
diff --git a/Engine/Source/Programs/UnrealTraceServer/src/Asio.h b/Engine/Source/Programs/UnrealTraceServer/src/Asio.h
new file mode 100644
--- /dev/null
+++ b/Engine/Source/Programs/UnrealTraceServer/src/Asio.h
@@ -0,0 +1,56 @@
+// Copyright Epic Games, Inc. All Rights Reserved.
+
+#pragma once
+
+#if defined(__clang__)
+#	if defined(__APPLE__)
+#		define TS_ASIO_INCLUDE_BEGIN \
+			_Pragma("GCC diagnostic push") \
+			_Pragma("GCC diagnostic ignored \"-Wdeprecated-declarations\"")
+#		define TS_ASIO_INCLUDE_END \
+			_Pragma("GCC diagnostic pop")
+#	endif
+#else
+#	define TS_ASIO_INCLUDE_BEGIN
+#	define TS_ASIO_INCLUDE_END
+#endif
+
+#define ASIO_SEPARATE_COMPILATION
+#define ASIO_STANDALONE
+#define ASIO_NO_EXCEPTIONS
+#define ASIO_NO_TYPEID
+#define ASIO_DISABLE_NOEXCEPT
+
+namespace asio {
+namespace detail {
+
+template <typename ExceptionType>
+void throw_exception(const ExceptionType& Exception)
+{
+	/* Intentionally blank */
+}
+
+} // namespace detail
+} // namespace asio
+
+TS_ASIO_INCLUDE_BEGIN
+#include "asio/connect.hpp"
+#include "asio/error.hpp"
+#include "asio/io_context.hpp"
+#include "asio/ip/tcp.hpp"
+#include "asio/read.hpp"
+#include "asio/write.hpp"
+#include "asio/write_at.hpp"
+
+#if TS_USING(TS_PLATFORM_WINDOWS)
+#	include "asio/windows/random_access_handle.hpp"
+#	include "asio/windows/object_handle.hpp"
+#endif
+TS_ASIO_INCLUDE_END
+
+#if TS_USING(TS_PLATFORM_WINDOWS)
+#	include <winsock2.h>
+#	pragma comment(lib, "ws2_32.lib")
+#endif
+
+/* vim: set noexpandtab : */
diff --git a/Engine/Source/Programs/UnrealTraceServer/src/AsioContext.cpp b/Engine/Source/Programs/UnrealTraceServer/src/AsioContext.cpp
new file mode 100644
--- /dev/null
+++ b/Engine/Source/Programs/UnrealTraceServer/src/AsioContext.cpp
@@ -0,0 +1,79 @@
+// Copyright Epic Games, Inc. All Rights Reserved.
+
+#include "Pch.h"
+#include "AsioContext.h"
+
+////////////////////////////////////////////////////////////////////////////////
+FAsioContext::FAsioContext(int32 ThreadCount)
+{
+	if (ThreadCount < 1)
+	{
+		ThreadCount = 1;
+	}
+
+	IoContext = new asio::io_context(ThreadCount);
+
+	ThreadPool.SetNum(ThreadCount);
+}
+
+////////////////////////////////////////////////////////////////////////////////
+FAsioContext::~FAsioContext()
+{
+	Wait();
+	delete IoContext;
+}
+
+////////////////////////////////////////////////////////////////////////////////
+void FAsioContext::Start()
+{
+	if (bRunning)
+	{
+		return;
+	}
+
+	for (std::thread& Thread : ThreadPool)
+	{
+		std::thread TempThread([this] () { IoContext->run(); });
+		Thread = MoveTemp(TempThread);
+	}
+
+	bRunning = true;
+}
+
+////////////////////////////////////////////////////////////////////////////////
+void FAsioContext::Stop(bool bWait)
+{
+	if (!bRunning)
+	{
+		return;
+	}
+
+	IoContext->stop();
+
+	if (bWait)
+	{
+		Wait();
+	}
+}
+
+////////////////////////////////////////////////////////////////////////////////
+void FAsioContext::Wait()
+{
+	for (std::thread& Thread : ThreadPool)
+	{
+		if (Thread.joinable())
+		{
+			Thread.join();
+		}
+	}
+
+	bRunning = false;
+}
+
+////////////////////////////////////////////////////////////////////////////////
+asio::io_context& FAsioContext::Get()
+{
+	return *IoContext;
+}
+
+/* vim: set noexpandtab : */
diff --git a/Engine/Source/Programs/UnrealTraceServer/src/AsioContext.h b/Engine/Source/Programs/UnrealTraceServer/src/AsioContext.h
new file mode 100644
--- /dev/null
+++ b/Engine/Source/Programs/UnrealTraceServer/src/AsioContext.h
@@ -0,0 +1,27 @@
+// Copyright Epic Games, Inc. All Rights Reserved.
+
+#pragma once
+
+#include "Asio.h"
+#include "Foundation.h"
+
+#include <thread>
+
+////////////////////////////////////////////////////////////////////////////////
+class FAsioContext
+{
+public:
+							FAsioContext(int32 ThreadCount);
+							~FAsioContext();
+	asio::io_context&		Get();
+	void					Start();
+	void					Stop(bool bWait=false);
+	void					Wait();
+
+private:
+	asio::io_context*		IoContext;
+	TArray<std::thread>		ThreadPool;
+	bool					bRunning = false;
+};
+
+/* vim: set noexpandtab : */
diff --git a/Engine/Source/Programs/UnrealTraceServer/src/AsioFile.cpp b/Engine/Source/Programs/UnrealTraceServer/src/AsioFile.cpp
new file mode 100644
--- /dev/null
+++ b/Engine/Source/Programs/UnrealTraceServer/src/AsioFile.cpp
@@ -0,0 +1,144 @@
+// Copyright Epic Games, Inc. All Rights Reserved.
+
+#include "Pch.h"
+#include "AsioFile.h"
+#include "Foundation.h"
+
+#if !TS_USING(TS_PLATFORM_WINDOWS)
+#include <sys/file.h>
+#include <errno.h>
+#endif
+
+////////////////////////////////////////////////////////////////////////////////
+FAsioFile::FAsioFile(asio::io_context& IoContext, uintptr_t OsHandle)
+#if TS_USING(TS_PLATFORM_WINDOWS)
+: Handle(IoContext, HANDLE(OsHandle))
+#else
+: StreamDescriptor(IoContext, OsHandle)
+#endif
+{
+}
+
+////////////////////////////////////////////////////////////////////////////////
+bool FAsioFile::IsOpen() const
+{
+#if TS_USING(TS_PLATFORM_WINDOWS)
+	return Handle.is_open();
+#else
+	return StreamDescriptor.is_open();
+#endif
+}
+
+////////////////////////////////////////////////////////////////////////////////
+void FAsioFile::Close()
+{
+#if TS_USING(TS_PLATFORM_WINDOWS)
+	Handle.close();
+#else
+	StreamDescriptor.close();
+#endif
+}
+
+////////////////////////////////////////////////////////////////////////////////
+bool FAsioFile::Write(const void* Src, uint32 Size, FAsioIoSink* Sink, uint32 Id)
+{
+	if (!SetSink(Sink, Id))
+	{
+		return false;
+	}
+
+#if TS_USING(TS_PLATFORM_WINDOWS)
+	asio::async_write_at(
+		Handle,
+		Offset,
+#else
+	asio::async_write(
+		StreamDescriptor,
+#endif
+		asio::buffer(Src, Size),
+		[this] (const asio::error_code& ErrorCode, size_t BytesWritten)
+		{
+			Offset += BytesWritten;
+			return OnIoComplete(ErrorCode, uint32(BytesWritten));
+		}
+	);
+
+	return true;
+}
+
+////////////////////////////////////////////////////////////////////////////////
+bool FAsioFile::Read(void* Dest, uint32 Size, FAsioIoSink* Sink, uint32 Id)
+{
+	return ReadSome(Dest, Size, Sink, Id);
+}
+
+////////////////////////////////////////////////////////////////////////////////
+bool FAsioFile::ReadSome(void* Dest, uint32 DestSize, FAsioIoSink* Sink, uint32 Id)
+{
+	if (!SetSink(Sink, Id))
+	{
+		return false;
+	}
+
+#if TS_USING(TS_PLATFORM_WINDOWS)
+	Handle.async_read_some_at(
+		Offset,
+#else
+	StreamDescriptor.async_read_some(
+#endif
+		asio::buffer(Dest, DestSize),
+		[this] (const asio::error_code& ErrorCode, size_t BytesRead)
+		{
+			Offset += BytesRead;
+			return OnIoComplete(ErrorCode, uint32(BytesRead));
+		}
+	);
+
+	return true;
+}
+
+////////////////////////////////////////////////////////////////////////////////
+FAsioWriteable* FAsioFile::WriteFile(asio::io_context& IoContext, const char* Path)
+{
+#if TS_USING(TS_PLATFORM_WINDOWS)
+	FWinApiStr PathW(Path);
+	HANDLE Handle = CreateFileW(PathW, GENERIC_WRITE, FILE_SHARE_READ, nullptr,
+		CREATE_ALWAYS, FILE_FLAG_OVERLAPPED|FILE_ATTRIBUTE_NORMAL, nullptr);
+	if (Handle == INVALID_HANDLE_VALUE)
+	{
+		return nullptr;
+	}
+	return new FAsioFile(IoContext, uintptr_t(Handle));
+#else
+	int File = open(Path, O_WRONLY|O_CREAT, 0666);
+	if (!File)
+	{
+		return nullptr;
+	}
+	return new FAsioFile(IoContext, uintptr_t(File));
+#endif
+}
+
+////////////////////////////////////////////////////////////////////////////////
+FAsioReadable* FAsioFile::ReadFile(asio::io_context& IoContext, const char* Path)
+{
+#if TS_USING(TS_PLATFORM_WINDOWS)
+	FWinApiStr PathW(Path);
+	HANDLE Handle = CreateFileW(PathW, GENERIC_READ, FILE_SHARE_READ|FILE_SHARE_WRITE,
+		nullptr, OPEN_EXISTING, FILE_FLAG_OVERLAPPED, nullptr);
+	if (Handle == INVALID_HANDLE_VALUE)
+	{
+		return nullptr;
+	}
+	return new FAsioFile(IoContext, uintptr_t(Handle));
+#else
+	int File = open(Path, O_RDONLY, 0444);
+	if (!File)
+	{
+		return nullptr;
+	}
+	return new FAsioFile(IoContext, uintptr_t(File));
+#endif
+}
+
+/* vim: set noexpandtab : */
diff --git a/Engine/Source/Programs/UnrealTraceServer/src/AsioFile.h b/Engine/Source/Programs/UnrealTraceServer/src/AsioFile.h
new file mode 100644
--- /dev/null
+++ b/Engine/Source/Programs/UnrealTraceServer/src/AsioFile.h
@@ -0,0 +1,36 @@
+// Copyright Epic Games, Inc. All Rights Reserved.
+
+#pragma once
+
+#include "Asio.h"
+
+#include "AsioIoable.h"
+#if !TS_USING(TS_PLATFORM_WINDOWS)
+#include "asio/posix/stream_descriptor.hpp"
+#endif
+
+////////////////////////////////////////////////////////////////////////////////
+class FAsioFile
+	: public FAsioReadable
+	, public FAsioWriteable
+{
+public:
+										FAsioFile(asio::io_context& IoContext, uintptr_t OsHandle);
+	static FAsioWriteable*				WriteFile(asio::io_context& IoContext, const char* Path);
+	static FAsioReadable*				ReadFile(asio::io_context& IoContext, const char* Path);
+	virtual bool						IsOpen() const override;
+	virtual void						Close() override;
+	virtual bool						Write(const void* Src, uint32 Size, FAsioIoSink* Sink, uint32 Id) override;
+	virtual bool						Read(void* Dest, uint32 Size, FAsioIoSink* Sink, uint32 Id) override;
+	virtual bool						ReadSome(void* Dest, uint32 BufferSize, FAsioIoSink* Sink, uint32 Id) override;
+
+private:
+#if TS_USING(TS_PLATFORM_WINDOWS)
+	asio::windows::random_access_handle	Handle;
+#else
+	asio::posix::stream_descriptor		StreamDescriptor;
+#endif
+	uint64								Offset = 0;
+};
+
+/* vim: set noexpandtab : */
diff --git a/Engine/Source/Programs/UnrealTraceServer/src/AsioIoable.cpp b/Engine/Source/Programs/UnrealTraceServer/src/AsioIoable.cpp
new file mode 100644
--- /dev/null
+++ b/Engine/Source/Programs/UnrealTraceServer/src/AsioIoable.cpp
@@ -0,0 +1,44 @@
+// Copyright Epic Games, Inc. All Rights Reserved.
+
+#include "Pch.h"
+#include "AsioIoable.h"
+
+////////////////////////////////////////////////////////////////////////////////
+bool FAsioIoable::SetSink(FAsioIoSink* Ptr, uint32 Id)
+{
+	if (SinkPtr != nullptr)
+	{
+		return false;
+	}
+
+	SinkPtr = Ptr;
+	SinkId = Id;
+	return true;
+}
+
+////////////////////////////////////////////////////////////////////////////////
+void FAsioIoable::OnIoComplete(const asio::error_code& ErrorCode, int32 Size)
+{
+	if (SinkPtr == nullptr)
+	{
+		return;
+	}
+
+#if TS_USING(TS_BUILD_DEBUG) && 0
+	std::string ErrorMessage;
+	{
+		ErrorMessage = ErrorCode.message();
+	}
+#endif
+
+	if (ErrorCode)
+	{
+		Size = 0 - ErrorCode.value();
+	}
+
+	FAsioIoSink* Ptr = SinkPtr;
+	SinkPtr = nullptr;
+	Ptr->OnIoComplete(SinkId, Size);
+}
+
+/* vim: set noexpandtab : */
diff --git a/Engine/Source/Programs/UnrealTraceServer/src/AsioIoable.h b/Engine/Source/Programs/UnrealTraceServer/src/AsioIoable.h
new file mode 100644
--- /dev/null
+++ b/Engine/Source/Programs/UnrealTraceServer/src/AsioIoable.h
@@ -0,0 +1,50 @@
+// Copyright Epic Games, Inc. All Rights Reserved.
+
+#pragma once
+
+#include "Asio.h"
+
+////////////////////////////////////////////////////////////////////////////////
+class FAsioIoSink
+{
+public:
+	virtual void	OnIoComplete(uint32 Id, int32 Size) = 0;
+};
+
+////////////////////////////////////////////////////////////////////////////////
+class FAsioIoable
+{
+public:
+	virtual			~FAsioIoable() = default;
+	virtual bool	IsOpen() const = 0;
+	virtual void	Close() = 0;
+
+protected:
+	bool			SetSink(FAsioIoSink* Ptr, uint32 Id);
+	void			OnIoComplete(const asio::error_code& ErrorCode, int32 Size);
+
+private:
+	FAsioIoSink*	SinkPtr = nullptr;
+	uint32			SinkId;
+};
+
+////////////////////////////////////////////////////////////////////////////////
+class FAsioReadable
+	: public virtual FAsioIoable
+{
+public:
+	virtual bool	Read(void* Dest, uint32 Size, FAsioIoSink* Sink, uint32 Id) = 0;
+	virtual bool	ReadSome(void* Dest, uint32 DestSize, FAsioIoSink* Sink, uint32 Id) = 0;
+};
+
+
+
+////////////////////////////////////////////////////////////////////////////////
+class FAsioWriteable
+	: public virtual FAsioIoable
+{
+public:
+	virtual bool	Write(const void* Src, uint32 Size, FAsioIoSink* Sink, uint32 Id) = 0;
+};
+
+/* vim: set noexpandtab : */
diff --git a/Engine/Source/Programs/UnrealTraceServer/src/AsioSocket.cpp b/Engine/Source/Programs/UnrealTraceServer/src/AsioSocket.cpp
new file mode 100644
--- /dev/null
+++ b/Engine/Source/Programs/UnrealTraceServer/src/AsioSocket.cpp
@@ -0,0 +1,133 @@
+// Copyright Epic Games, Inc. All Rights Reserved.
+
+#include "Pch.h"
+#include "AsioSocket.h"
+#include "Foundation.h"
+
+////////////////////////////////////////////////////////////////////////////////
+FAsioSocket::FAsioSocket(asio::ip::tcp::socket& InSocket)
+: Socket(MoveTemp(InSocket))
+{
+}
+
+////////////////////////////////////////////////////////////////////////////////
+FAsioSocket::~FAsioSocket()
+{
+	check(!IsOpen());
+}
+
+////////////////////////////////////////////////////////////////////////////////
+asio::io_context& FAsioSocket::GetIoContext()
+{
+	return Socket.get_executor().context();
+}
+
+////////////////////////////////////////////////////////////////////////////////
+bool FAsioSocket::IsOpen() const
+{
+	return Socket.is_open();
+}
+
+////////////////////////////////////////////////////////////////////////////////
+uint32 FAsioSocket::GetRemoteAddress() const
+{
+	if (!IsOpen())
+	{
+		return 0;
+	}
+
+	const asio::ip::tcp::endpoint Endpoint = Socket.remote_endpoint();
+	asio::ip::address Address = Endpoint.address();
+	return Address.is_v4() ? Address.to_v4().to_uint() : 0;
+}
+
+////////////////////////////////////////////////////////////////////////////////
+uint32 FAsioSocket::GetRemotePort() const
+{
+	if (!IsOpen())
+	{
+		return 0;
+	}
+
+	const asio::ip::tcp::endpoint Endpoint = Socket.remote_endpoint();
+	return Endpoint.port();
+}
+
+////////////////////////////////////////////////////////////////////////////////
+uint32 FAsioSocket::GetLocalPort() const
+{
+	if (!IsOpen())
+	{
+		return 0;
+	}
+
+	const asio::ip::tcp::endpoint Endpoint = Socket.local_endpoint();
+	return Endpoint.port();
+}
+
+////////////////////////////////////////////////////////////////////////////////
+void FAsioSocket::Close()
+{
+	Socket.close();
+}
+
+////////////////////////////////////////////////////////////////////////////////
+bool FAsioSocket::Read(void* Dest, uint32 Size, FAsioIoSink* Sink, uint32 Id)
+{
+	if (!SetSink(Sink, Id))
+	{
+		return false;
+	}
+
+	asio::async_read(
+		Socket,
+		asio::buffer(Dest, Size),
+		[this] (const asio::error_code& ErrorCode, size_t BytesReceived)
+		{
+			OnIoComplete(ErrorCode, uint32(BytesReceived));
+		}
+	);
+
+	return true;
+}
+
+////////////////////////////////////////////////////////////////////////////////
+bool FAsioSocket::ReadSome(void* Dest, uint32 BufferSize, FAsioIoSink* Sink, uint32 Id)
+{
+	if (!SetSink(Sink, Id))
+	{
+		return false;
+	}
+
+	Socket.async_receive(
+		asio::buffer(Dest, BufferSize),
+		[this] (const asio::error_code& ErrorCode, size_t BytesReceived)
+		{
+			return OnIoComplete(ErrorCode, uint32(BytesReceived));
+		}
+	);
+
+	return true;
+}
+
+////////////////////////////////////////////////////////////////////////////////
+bool FAsioSocket::Write(const void* Src, uint32 Size, FAsioIoSink* Sink, uint32 Id)
+{
+	if (!SetSink(Sink, Id))
+	{
+		return false;
+	}
+
+	asio::async_write(
+		Socket,
+		asio::buffer(Src, Size),
+		[this] (const asio::error_code& ErrorCode, size_t BytesSent)
+		{
+			return OnIoComplete(ErrorCode, uint32(BytesSent));
+		}
+	);
+
+	return true;
+}
+
+/* vim: set noexpandtab : */
diff --git a/Engine/Source/Programs/UnrealTraceServer/src/AsioSocket.h b/Engine/Source/Programs/UnrealTraceServer/src/AsioSocket.h
new file mode 100644
--- /dev/null
+++ b/Engine/Source/Programs/UnrealTraceServer/src/AsioSocket.h
@@ -0,0 +1,30 @@
+// Copyright Epic Games, Inc. All Rights Reserved.
+
+#pragma once
+
+#include "Asio.h"
+#include "AsioIoable.h"
+
+////////////////////////////////////////////////////////////////////////////////
+class FAsioSocket
+	: public FAsioReadable
+	, public FAsioWriteable
+{
+public:
+							FAsioSocket(asio::ip::tcp::socket& InSocket);
+	virtual					~FAsioSocket();
+	asio::io_context&		GetIoContext();
+	virtual bool			IsOpen() const override;
+	virtual void			Close() override;
+	uint32					GetRemoteAddress() const;
+	uint32					GetRemotePort() const;
+	uint32					GetLocalPort() const;
+	virtual bool			Read(void* Dest, uint32 Size, FAsioIoSink* Sink, uint32 Id) override;
+	virtual bool			ReadSome(void* Dest, uint32 DestSize, FAsioIoSink* Sink, uint32 Id) override;
+	virtual bool			Write(const void* Src, uint32 Size, FAsioIoSink* Sink, uint32 Id) override;
+
+private:
+	asio::ip::tcp::socket	Socket;
+};
+
+/* vim: set noexpandtab : */
diff --git a/Engine/Source/Programs/UnrealTraceServer/src/AsioTcpServer.cpp b/Engine/Source/Programs/UnrealTraceServer/src/AsioTcpServer.cpp
new file mode 100644
--- /dev/null
+++ b/Engine/Source/Programs/UnrealTraceServer/src/AsioTcpServer.cpp
@@ -0,0 +1,129 @@
+// Copyright Epic Games, Inc. All Rights Reserved.
+
+#include "Pch.h"
+#include "AsioTcpServer.h"
+#include "Foundation.h"
+
+////////////////////////////////////////////////////////////////////////////////
+FAsioTcpServer::FAsioTcpServer(asio::io_context& IoContext)
+: Acceptor(IoContext)
+{
+}
+
+////////////////////////////////////////////////////////////////////////////////
+FAsioTcpServer::~FAsioTcpServer()
+{
+	check(!IsOpen());
+}
+
+////////////////////////////////////////////////////////////////////////////////
+bool FAsioTcpServer::IsOpen() const
+{
+	return Acceptor.is_open();
+}
+
+////////////////////////////////////////////////////////////////////////////////
+asio::io_context& FAsioTcpServer::GetIoContext()
+{
+	return Acceptor.get_executor().context();
+}
+
+////////////////////////////////////////////////////////////////////////////////
+void FAsioTcpServer::Close()
+{
+	Acceptor.close();
+}
+
+////////////////////////////////////////////////////////////////////////////////
+uint32 FAsioTcpServer::GetPort() const
+{
+	using asio::ip::tcp;
+
+	if (!Acceptor.is_open())
+	{
+		return 0;
+	}
+
+	asio::error_code ErrorCode;
+	tcp::endpoint Endpoint = Acceptor.local_endpoint(ErrorCode);
+	if (ErrorCode)
+	{
+		return 0;
+	}
+
+	return Endpoint.port();
+}
+
+////////////////////////////////////////////////////////////////////////////////
+bool FAsioTcpServer::StartServer(uint32 Port, uint32 Backlog)
+{
+	if (Acceptor.is_open())
+	{
+		return false;
+	}
+
+	using asio::ip::tcp;
+
+	tcp::acceptor TempAcceptor(GetIoContext());
+
+#if TS_USING(TS_PLATFORM_WINDOWS)
+	DWORD Flags = WSA_FLAG_NO_HANDLE_INHERIT|WSA_FLAG_OVERLAPPED;
+	SOCKET Socket = WSASocketW(AF_INET, SOCK_STREAM, IPPROTO_TCP, nullptr, 0, Flags);
+	TempAcceptor.assign(tcp::v4(), Socket);
+#else
+	TempAcceptor.open(tcp::v4());
+#endif
+
+	tcp::endpoint Endpoint(tcp::v4(), uint16(Port));
+	TempAcceptor.bind(Endpoint);
+	TempAcceptor.listen(Backlog);
+
+	tcp::endpoint LocalEndpoint = TempAcceptor.local_endpoint();
+	if (LocalEndpoint.port() == 0)
+	{
+		return false;
+	}
+
+	Acceptor = MoveTemp(TempAcceptor);
+	AsyncAccept();
+	return true;
+}
+
+////////////////////////////////////////////////////////////////////////////////
+bool FAsioTcpServer::StopServer()
+{
+	if (!Acceptor.is_open())
+	{
+		return false;
+	}
+
+	Acceptor.close();
+	return true;
+}
+
+////////////////////////////////////////////////////////////////////////////////
+void FAsioTcpServer::AsyncAccept()
+{
+	using asio::ip::tcp;
+
+	Acceptor.async_accept([this] (
+		const asio::error_code& ErrorCode,
+		tcp::socket&& Socket)
+	{
+#if TS_USING(TS_PLATFORM_WINDOWS)
+		// This is a best effort and likely not to work because WFP filters (e.g.
+		// antimalware or firewalls) can redirect socket handles.
+		SetHandleInformation(HANDLE(SOCKET(Socket.native_handle())), HANDLE_FLAG_INHERIT, 0);
+#endif
+
+		if (!ErrorCode && OnAccept(Socket))
+		{
+			AsyncAccept();
+			return;
+		}
+
+		Close();
+	});
+}
+
+/* vim: set noexpandtab : */
diff --git a/Engine/Source/Programs/UnrealTraceServer/src/AsioTcpServer.h b/Engine/Source/Programs/UnrealTraceServer/src/AsioTcpServer.h
new file mode 100644
--- /dev/null
+++ b/Engine/Source/Programs/UnrealTraceServer/src/AsioTcpServer.h
@@ -0,0 +1,26 @@
+// Copyright Epic Games, Inc. All Rights Reserved.
+
+#pragma once
+
+#include "Asio.h"
+
+////////////////////////////////////////////////////////////////////////////////
+class FAsioTcpServer
+{
+public:
+							FAsioTcpServer(asio::io_context& IoContext);
+	virtual					~FAsioTcpServer();
+	asio::io_context&		GetIoContext();
+	bool					IsOpen() const;
+	void					Close();
+	virtual bool			OnAccept(asio::ip::tcp::socket& Socket) = 0;
+	uint32					GetPort() const;
+	bool					StartServer(uint32 Port=0, uint32 Backlog=1);
+	bool					StopServer();
+
+private:
+	void					AsyncAccept();
+	asio::ip::tcp::acceptor	Acceptor;
+};
+
+/* vim: set noexpandtab : */
diff --git a/Engine/Source/Programs/UnrealTraceServer/src/AsioTickable.cpp b/Engine/Source/Programs/UnrealTraceServer/src/AsioTickable.cpp
new file mode 100644
--- /dev/null
+++ b/Engine/Source/Programs/UnrealTraceServer/src/AsioTickable.cpp
@@ -0,0 +1,100 @@
+// Copyright Epic Games, Inc. All Rights Reserved.
+
+#include "Pch.h"
+#include "AsioTickable.h"
+#include "Foundation.h"
+
+#include <chrono>
+
+////////////////////////////////////////////////////////////////////////////////
+FAsioTickable::FAsioTickable(asio::io_context& IoContext)
+: Timer(IoContext)
+{
+}
+
+////////////////////////////////////////////////////////////////////////////////
+FAsioTickable::~FAsioTickable()
+{
+	check(!IsActive());
+}
+
+////////////////////////////////////////////////////////////////////////////////
+asio::io_context& FAsioTickable::GetIoContext()
+{
+	return Timer.get_executor().context();
+}
+
+////////////////////////////////////////////////////////////////////////////////
+bool FAsioTickable::IsActive() const
+{
+	return MillisecondRate != 0;
+}
+
+////////////////////////////////////////////////////////////////////////////////
+bool FAsioTickable::StartTick(uint32 InMillisecondRate)
+{
+	if (MillisecondRate)
+	{
+		return false;
+	}
+
+	MillisecondRate = InMillisecondRate;
+	if (MillisecondRate)
+	{
+		AsyncTick();
+		return true;
+	}
+
+	return false;
+}
+
+////////////////////////////////////////////////////////////////////////////////
+bool FAsioTickable::StopTick()
+{
+	MillisecondRate = 0;
+	Timer.cancel();
+	return true;
+}
+
+////////////////////////////////////////////////////////////////////////////////
+bool FAsioTickable::TickOnce(uint32 InMillisecondRate)
+{
+	if (MillisecondRate)
+	{
+		return false;
+	}
+
+	MillisecondRate = InMillisecondRate;
+	if (MillisecondRate)
+	{
+		AsyncTick();
+		MillisecondRate = 0;
+		return true;
+	}
+
+	return false;
+}
+
+////////////////////////////////////////////////////////////////////////////////
+void FAsioTickable::AsyncTick()
+{
+	auto StdTime = std::chrono::milliseconds(MillisecondRate);
+	Timer.expires_after(StdTime);
+
+	Timer.async_wait([this] (const asio::error_code& ErrorCode)
+	{
+		if (ErrorCode)
+		{
+			return;
+		}
+
+		OnTick();
+
+		if (MillisecondRate)
+		{
+			AsyncTick();
+		}
+	});
+}
+
+/* vim: set noexpandtab : */
diff --git a/Engine/Source/Programs/UnrealTraceServer/src/AsioTickable.h b/Engine/Source/Programs/UnrealTraceServer/src/AsioTickable.h
new file mode 100644
--- /dev/null
+++ b/Engine/Source/Programs/UnrealTraceServer/src/AsioTickable.h
@@ -0,0 +1,26 @@
+// Copyright Epic Games, Inc. All Rights Reserved.
+
+#pragma once
+
+#include "Asio.h"
+
+////////////////////////////////////////////////////////////////////////////////
+class FAsioTickable
+{
+public:
+						FAsioTickable(asio::io_context& IoContext);
+	virtual				~FAsioTickable();
+	asio::io_context&	GetIoContext();
+	bool				IsActive() const;
+	bool				StartTick(uint32 MillisecondRate);
+	bool				StopTick();
+	bool				TickOnce(uint32 MillisecondRate);
+	virtual void		OnTick() = 0;
+
+private:
+	void				AsyncTick();
+	asio::steady_timer	Timer;
+	uint32				MillisecondRate = 0;
+};
+
+/* vim: set noexpandtab : */
diff --git a/Engine/Source/Programs/UnrealTraceServer/src/Cbor.cpp b/Engine/Source/Programs/UnrealTraceServer/src/Cbor.cpp
new file mode 100644
--- /dev/null
+++ b/Engine/Source/Programs/UnrealTraceServer/src/Cbor.cpp
@@ -0,0 +1,681 @@
+// Copyright Epic Games, Inc. All Rights Reserved.
+
+#include "Pch.h"
+#include "Cbor.h"
+#include "Utils.h"
+
+// {{{1 misc -------------------------------------------------------------------
+
+enum ECborMajorType
+{
+	Integer			= 0 << 5,
+	NegativeInteger	= 1 << 5,
+	ByteString		= 2 << 5,
+	TextString		= 3 << 5,
+	Array			= 4 << 5,
+	Map				= 5 << 5,
+	Tag				= 6 << 5,
+	Float_Simple	= 7 << 5,
+};
+
+
+
+// {{{1 context ----------------------------------------------------------------
+
+////////////////////////////////////////////////////////////////////////////////
+void FCborContext::Reset()
+{
+	NextCursor = nullptr;
+}
+
+////////////////////////////////////////////////////////////////////////////////
+ECborType FCborContext::GetType() const
+{
+	return Type;
+}
+
+////////////////////////////////////////////////////////////////////////////////
+int64 FCborContext::GetLength() const
+{
+	return Param;
+}
+
+////////////////////////////////////////////////////////////////////////////////
+FStringView	FCborContext::AsString() const
+{
+	if (GetType() != ECborType::String)
+	{
+		return FStringView();
+	}
+
+    return FStringView((const char*)(NextCursor - Param), Param);
+}
+
+////////////////////////////////////////////////////////////////////////////////
+int64 FCborContext::AsInteger() const
+{
+    return Param;
+}
+
+////////////////////////////////////////////////////////////////////////////////
+void FCborContext::SetError()
+{
+	Type = ECborType::Error;
+	Param = -1;
+	NextCursor = nullptr;
+}
+
+
+
+// {{{1 writer -----------------------------------------------------------------
+
+////////////////////////////////////////////////////////////////////////////////
+FCborWriter::FCborWriter(FInlineBuffer& InBuffer)
+: Buffer(InBuffer)
+{
+}
+
+////////////////////////////////////////////////////////////////////////////////
+void FCborWriter::OpenMap(int32 ItemCount)
+{
+	if (ItemCount != -1)
+	{
+		WriteParam(ECborMajorType::Map, ItemCount);
+		return;
+	}
+
+	uint8 Initial = ECborMajorType::Map | 31;
+	Buffer.Append(&Initial, sizeof(Initial));
+}
+
+////////////////////////////////////////////////////////////////////////////////
+void FCborWriter::OpenArray(int32 ItemCount)
+{
+	if (ItemCount != -1)
+	{
+		WriteParam(ECborMajorType::Array, ItemCount);
+		return;
+	}
+
+	uint8 Initial = ECborMajorType::Array | 31;
+	Buffer.Append(&Initial, sizeof(Initial));
+}
+
+////////////////////////////////////////////////////////////////////////////////
+void FCborWriter::Close()
+{
+	uint8 Initial = ECborMajorType::Float_Simple | 31;
+	Buffer.Append(&Initial, sizeof(Initial));
+}
+
+////////////////////////////////////////////////////////////////////////////////
+void FCborWriter::WriteString(const char* Value, int32 Length)
+{
+	Length = (Length < 0) ? int32(strlen(Value)) : Length;
+
+	WriteParam(ECborMajorType::TextString, Length);
+	Buffer.Append(Value, Length);
+}
+
+////////////////////////////////////////////////////////////////////////////////
+void FCborWriter::WriteInteger(int64 Value)
+{
+	uint32 Negative = (Value < 0);
+	if (Negative)
+	{
+		Value = ~Value;
+	}
+
+	uint32 MajorType = Negative ? ECborMajorType::NegativeInteger : ECborMajorType::Integer;
+	WriteParam(MajorType, Value);
+}
+
+////////////////////////////////////////////////////////////////////////////////
+void FCborWriter::WriteParam(uint32 MajorType, uint64 Value)
+{
+	if (Value <= 23)
+	{
+		uint8 Initial = MajorType | uint8(Value);
+		Buffer.Append(&Initial, 1);
+		return;
+	}
+
+	uint32 BytesPow2 = 3;
+	BytesPow2 -= int32(Value <= 0xffff'ffffu);
+	BytesPow2 -= int32(Value <= 0x0000'ffffu);
+	BytesPow2 -= int32(Value <= 0x0000'00ffu);
+
+	Value = ((Value & 0xffff'ffff'0000'0000ull) >> 32) | ((Value & 0x0000'0000'ffff'ffffull) << 32);
+	Value = ((Value & 0xffff'0000'ffff'0000ull) >> 16) | ((Value & 0x0000'ffff'0000'ffffull) << 16);
+	Value = ((Value & 0xff00'ff00'ff00'ff00ull) >>  8) | ((Value & 0x00ff'00ff'00ff'00ffull) <<  8);
+
+	Value >>= 64 - (8 << BytesPow2);
+
+	struct {
+		uint8	Initial;
+		uint8	Payload[sizeof(Value)];
+	} Packet = {
+		uint8(MajorType | (24 + BytesPow2)),
+	};
+
+	memcpy(Packet.Payload, &Value, sizeof(Value));
+	Buffer.Append(&Packet, 1 + (1 << BytesPow2));
+}
+
+
+
+// {{{1 reader -----------------------------------------------------------------
+
+////////////////////////////////////////////////////////////////////////////////
+FCborReader::FCborReader(const uint8* Data, uint32 DataSize)
+: Base(Data)
+, DataEnd(Data + DataSize)
+{
+}
+
+////////////////////////////////////////////////////////////////////////////////
+int64 FCborReader::ReadParam(const uint8*& Cursor)
+{
+	uint32 Additional = *Cursor & 0b0001'1111;
+	++Cursor;
+
+	// [0-23] literal values
+	if (uint32(23 - Additional) <= 23)
+	{
+		return Additional;
+	}
+
+	// [24-27] value in subsequent bytes
+	else if (uint32(27 - Additional) <= 27)
+	{
+		uint32 Bytes = 1 << (Additional - 24);
+		if (Cursor + Bytes > DataEnd)
+		{
+			return -2;
+		}
+
+		int64 Value = 0;
+		do
+		{
+			Value <<= 8;
+			Value |= *Cursor;
+			++Cursor;
+		}
+		while (--Bytes);
+
+		return Value;
+	}
+
+	// 31: indeterminate length
+	else if (Additional == 31)
+	{
+		return -1;
+	}
+
+	// ...something went wrong
+	return -2;
+}
+
+////////////////////////////////////////////////////////////////////////////////
+bool FCborReader::ReadContainer(FCborContext& Context, const uint8* Cursor)
+{
+	bool bIsMap = (*Cursor >= ECborMajorType::Map);
+
+	int64 Param = ReadParam(Cursor);
+	if (Param <= -2)
+	{
+		Context.SetError();
+		return false;
+	}
+
+	Context.Type = bIsMap ? ECborType::Map : ECborType::Array;
+	Context.NextCursor = Cursor;
+	Context.Param = Param;
+	return true;
+}
+
+////////////////////////////////////////////////////////////////////////////////
+bool FCborReader::ReadInteger(FCborContext& Context, const uint8* Cursor)
+{
+	bool bIsNegative = (*Cursor >= ECborMajorType::NegativeInteger);
+
+	int64 Param = ReadParam(Cursor);
+	if (Param <= -1)
+	{
+		Context.SetError();
+		return false;
+	}
+
+	Context.Type = ECborType::Integer;
+	Context.NextCursor = Cursor;
+	Context.Param = bIsNegative ? ~Param : Param;
+	return true;
+}
+
+////////////////////////////////////////////////////////////////////////////////
+bool FCborReader::ReadString(FCborContext& Context, const uint8* Cursor)
+{
+	int64 Param = ReadParam(Cursor);
+	if (Param <= -2)
+	{
+		Context.SetError();
+		return false;
+	}
+	else if (Param >= 0)
+	{
+		if ((Cursor + Param) > DataEnd)
+		{
+			Context.SetError();
+			return false;
+		}
+
+		Cursor += Param;
+	}
+
+	Context.Type = ECborType::String;
+	Context.NextCursor = Cursor;
+	Context.Param = Param;
+	return true;
+}
+
+////////////////////////////////////////////////////////////////////////////////
+bool FCborReader::ReadNext(FCborContext& Context)
+{
+	const uint8* Cursor = (Context.NextCursor == nullptr) ? Base : Context.NextCursor;
+	if (Cursor >= DataEnd)
+	{
+		Context.Type = ECborType::Eof;
+		Context.Param = -1;
+		return false;
+	}
+
+	int64 Param;
+	switch (*Cursor & 0b1110'0000)
+	{
+	case ECborMajorType::Integer:
+	case ECborMajorType::NegativeInteger:
+		return ReadInteger(Context, Cursor);
+
+	case ECborMajorType::ByteString:
+	case ECborMajorType::TextString:
+		return ReadString(Context, Cursor);
+
+	case ECborMajorType::Array:
+	case ECborMajorType::Map:
+		return ReadContainer(Context, Cursor);
+	
+	case ECborMajorType::Tag:
+		/* unsupported */
+		break;
+
+	case ECborMajorType::Float_Simple:
+		Param = *Cursor & 0b0001'1111;
+		++Cursor;
+
+		switch (Param)
+		{
+		case 24:
+			if (Cursor >= DataEnd)
+			{
+				Context.SetError();
+				return false;
+			}
+
+			switch (*Cursor)
+			{
+			case 20: // false
+			case 21: // true
+			case 22: // null
+			case 23: // undefined
+				break;
+			}
+			break;
+
+		case 25: // half
+		case 26: // float
+		case 27: // double
+			/* if (Cursor + float_size > DataEnd) "basil!" / Context.SetError() */
+			break;
+
+		case 31:
+			Context.Type = ECborType::End;
+			Context.NextCursor = Cursor;
+			Context.Param = -1;
+			return false; // "false" to facilitate ReadNext() in loops
+		}
+		// 31 = break
+		break;
+	}
+
+	return false;
+}
+
+
+
+// {{{1 test-read --------------------------------------------------------------
+
+////////////////////////////////////////////////////////////////////////////////
+#if defined(_MSC_VER)
+#   define DEBUG_BREAK() do { __debugbreak(); } while (0)
+#else
+#   define DEBUG_BREAK() do { __builtin_trap(); } while (0)
+#endif
+
+#define REQUIRE(x) \
+	do { \
+		if (!(x)) \
+			DEBUG_BREAK(); \
+	} while (false)
+
+////////////////////////////////////////////////////////////////////////////////
+static void TestCborReader_Integer()
+{
+	struct
+	{
+		int64	Expected;
+		uint32	TruthSize;
+		uint8	Truth[9];
+	} Cases[] = {
+		{ 0,			1, { 0x00 } },
+		{ 0,			2, { 0x18, 0x00 } },
+		{ 0,			3, { 0x19, 0x00, 0x00 } },
+		{ 0,			5, { 0x1a, 0x00, 0x00, 0x00, 0x00 } },
+		{ 0,			9, { 0x1b, 0x00, 0x00, 0x00, 0x00, 0x00, 0x00, 0x00, 0x00 } },
+		{ 1,			1, { 0x01 } },
+		{ 23,			1, { 0x17 } },
+		{ 24,			2, { 0x18, 0x18 } },
+		{ 255,			2, { 0x18, 0xff } },
+		{ 256,			3, { 0x19, 0x01, 0x00 } },
+		{ 65535,		3, { 0x19, 0xff, 0xff } },
+		{ 65536,		5, { 0x1a, 0x00, 0x01, 0x00, 0x00 } },
+		{ 123456,		5, { 0x1a, 0x00, 0x01, 0xe2, 0x40 } },
+		{ 123456789,	5, { 0x1a, 0x07, 0x5b, 0xcd, 0x15 } },
+		{ -1,			1, { 0x20 } },
+		{ -2,			1, { 0x21 } },
+		{ -10,			1, { 0x29 } },
+		{ -24,			1, { 0x37 } },
+		{ -25,			2, { 0x38, 0x18 } },
+		{ -500,			3, { 0x39, 0x01, 0xf3 } },
+		{ -67000,		5, { 0x3a, 0x00, 0x01, 0x05, 0xb7 } },
+		{ -123456789,	5, { 0x3a, 0x07, 0x5b, 0xcd, 0x14 } },
+	};
+
+	for (const auto& Case : Cases)
+	{
+		FCborContext Context;
+		FCborReader Reader(Case.Truth, Case.TruthSize);
+		REQUIRE(Reader.ReadNext(Context));
+		REQUIRE(Context.GetType() == ECborType::Integer);
+		REQUIRE(Context.AsInteger() == Case.Expected);
+		REQUIRE(!Reader.ReadNext(Context));
+	}
+}
+
+////////////////////////////////////////////////////////////////////////////////
+static void TestCborReader_String()
+{
+	// "string"
+	{
+		const uint8 Truth[] = { 0x66, 0x73, 0x74, 0x72, 0x69, 0x6e, 0x67 };
+		const char* Expected = "string";
+
+		FCborContext Context;
+		FCborReader Reader(Truth, TS_ARRAY_COUNT(Truth));
+		REQUIRE(Reader.ReadNext(Context));
+		REQUIRE(Context.GetType() == ECborType::String);
+		REQUIRE(Context.AsString().Compare(Expected) == 0);
+		REQUIRE(!Reader.ReadNext(Context));
+	}
+
+	// ""
+	{
+		const uint8 Truth[] = { 0x60 };
+		const char* Expected = "";
+
+		FCborContext Context;
+		FCborReader Reader(Truth, TS_ARRAY_COUNT(Truth));
+		REQUIRE(Reader.ReadNext(Context));
+		REQUIRE(Context.GetType() == ECborType::String);
+		REQUIRE(Context.AsString().Compare(Expected) == 0);
+		REQUIRE(!Reader.ReadNext(Context));
+	}
+
+	// "stringstringstringstringstringstringstringstringstringstring"
+	{
+		const uint8 Truth[] = {
+			0x78, 0x3c,
+			0x73, 0x74, 0x72, 0x69, 0x6e, 0x67,
+			0x73, 0x74, 0x72, 0x69, 0x6e, 0x67,
+			0x73, 0x74, 0x72, 0x69, 0x6e, 0x67,
+			0x73, 0x74, 0x72, 0x69, 0x6e, 0x67,
+			0x73, 0x74, 0x72, 0x69, 0x6e, 0x67,
+			0x73, 0x74, 0x72, 0x69, 0x6e, 0x67,
+			0x73, 0x74, 0x72, 0x69, 0x6e, 0x67,
+			0x73, 0x74, 0x72, 0x69, 0x6e, 0x67,
+			0x73, 0x74, 0x72, 0x69, 0x6e, 0x67,
+			0x73, 0x74, 0x72, 0x69, 0x6e, 0x67
+		};
+		const char* Expected = "stringstringstringstringstringstringstringstringstringstring";
+
+		FCborContext Context;
+		FCborReader Reader(Truth, TS_ARRAY_COUNT(Truth));
+		REQUIRE(Reader.ReadNext(Context));
+		REQUIRE(Context.GetType() == ECborType::String);
+		REQUIRE(Context.AsString().Compare(Expected) == 0);
+		REQUIRE(!Reader.ReadNext(Context));
+	}
+}
+
+////////////////////////////////////////////////////////////////////////////////
+static void TestCborReader_Map()
+{
+	// {}
+	{
+		const uint8 Truth[] = { 0xa0 };
+
+		FCborContext Context;
+		FCborReader Reader(Truth, TS_ARRAY_COUNT(Truth));
+		REQUIRE(Reader.ReadNext(Context));
+		REQUIRE(Context.GetType() == ECborType::Map);
+		REQUIRE(Context.GetLength() == 0);
+		REQUIRE(!Reader.ReadNext(Context));
+	}
+
+	// {...}
+	{
+		const uint8 Truth[] = { 0xbf, 0xff };
+
+		FCborContext Context;
+		FCborReader Reader(Truth, TS_ARRAY_COUNT(Truth));
+		REQUIRE(Reader.ReadNext(Context));
+		REQUIRE(Context.GetType() == ECborType::Map);
+		REQUIRE(Context.GetLength() == -1);
+
+		REQUIRE(!Reader.ReadNext(Context));
+		REQUIRE(Context.GetType() == ECborType::End);
+
+		REQUIRE(!Reader.ReadNext(Context));
+		REQUIRE(Context.GetType() == ECborType::Eof);
+	}
+
+	// { "key":"value", 0:1 }
+	{
+		const uint8 Truth[] = {
+			0xa2,
+			0x63, 0x6b, 0x65, 0x79,
+			0x65, 0x76, 0x61, 0x6c, 0x75, 0x65,
+			0x00, 0x01
+		};
+
+		FCborContext Context;
+		FCborReader Reader(Truth, TS_ARRAY_COUNT(Truth));
+
+		REQUIRE(Reader.ReadNext(Context));
+		REQUIRE(Context.GetType() == ECborType::Map);
+		REQUIRE(Context.GetLength() == 2);
+
+		REQUIRE(Reader.ReadNext(Context));
+		REQUIRE(Context.GetType() == ECborType::String);
+		REQUIRE(Context.AsString().Compare("key") == 0);
+
+		REQUIRE(Reader.ReadNext(Context));
+		REQUIRE(Context.GetType() == ECborType::String);
+		REQUIRE(Context.AsString().Compare("value") == 0);
+
+		REQUIRE(Reader.ReadNext(Context));
+		REQUIRE(Context.GetType() == ECborType::Integer);
+		REQUIRE(Context.AsInteger() == 0);
+
+		REQUIRE(Reader.ReadNext(Context));
+		REQUIRE(Context.GetType() == ECborType::Integer);
+		REQUIRE(Context.AsInteger() == 1);
+
+		REQUIRE(!Reader.ReadNext(Context));
+		REQUIRE(Context.GetType() == ECborType::Eof);
+	}
+}
+
+// {{{1 test-write -------------------------------------------------------------
+
+////////////////////////////////////////////////////////////////////////////////
+static void TestCborWriter_Integer()
+{
+	struct
+	{
+		int64	Content;
+		uint32	TruthSize;
+		uint8	Truth[9];
+	} Cases[] = {
+		{ 0,			1, { 0x00 } },
+		{ 1,			1, { 0x01 } },
+		{ 23,			1, { 0x17 } },
+		{ 24,			2, { 0x18, 0x18 } },
+		{ 255,			2, { 0x18, 0xff } },
+		{ 256,			3, { 0x19, 0x01, 0x00 } },
+		{ 65535,		3, { 0x19, 0xff, 0xff } },
+		{ 65536,		5, { 0x1a, 0x00, 0x01, 0x00, 0x00 } },
+		{ 123456,		5, { 0x1a, 0x00, 0x01, 0xe2, 0x40 } },
+		{ 123456789,	5, { 0x1a, 0x07, 0x5b, 0xcd, 0x15 } },
+		{ -1,			1, { 0x20 } },
+		{ -2,			1, { 0x21 } },
+		{ -10,			1, { 0x29 } },
+		{ -24,			1, { 0x37 } },
+		{ -25,			2, { 0x38, 0x18 } },
+		{ -500,			3, { 0x39, 0x01, 0xf3 } },
+		{ -67000,		5, { 0x3a, 0x00, 0x01, 0x05, 0xb7 } },
+		{ -123456789,	5, { 0x3a, 0x07, 0x5b, 0xcd, 0x14 } },
+	};
+
+	for (const auto& Case : Cases)
+	{
+		TInlineBuffer<4> Buffer;
+		FCborWriter Writer(Buffer);
+		Writer.WriteInteger(Case.Content);
+		REQUIRE(Buffer->GetSize() == Case.TruthSize);
+		REQUIRE(memcmp(Buffer->GetData(), Case.Truth, Case.TruthSize) == 0);
+	}
+}
+
+////////////////////////////////////////////////////////////////////////////////
+static void TestCborWriter_String()
+{
+	struct
+	{
+		const char*	Value;
+		uint32		TruthSize;
+		const uint8	Truth[64];
+	} Cases[] = {
+		{ "string", 7, { 0x66, 0x73, 0x74, 0x72, 0x69, 0x6e, 0x67 } },
+		{ "",		1, { 0x60 } },
+		{
+			"stringstringstringstringstringstringstringstringstringstring",
+			62,
+			{ 0x78, 0x3c,
+			  0x73, 0x74, 0x72, 0x69, 0x6e, 0x67, 0x73, 0x74, 0x72, 0x69, 0x6e, 0x67,
+			  0x73, 0x74, 0x72, 0x69, 0x6e, 0x67, 0x73, 0x74, 0x72, 0x69, 0x6e, 0x67,
+			  0x73, 0x74, 0x72, 0x69, 0x6e, 0x67, 0x73, 0x74, 0x72, 0x69, 0x6e, 0x67,
+			  0x73, 0x74, 0x72, 0x69, 0x6e, 0x67, 0x73, 0x74, 0x72, 0x69, 0x6e, 0x67,
+			  0x73, 0x74, 0x72, 0x69, 0x6e, 0x67, 0x73, 0x74, 0x72, 0x69, 0x6e, 0x67 },
+		},
+	};
+
+	for (const auto& Case : Cases)
+	{
+		TInlineBuffer<8> Buffer;
+		FCborWriter Writer(Buffer);
+		Writer.WriteString(Case.Value);
+		REQUIRE(Buffer->GetSize() == Case.TruthSize);
+		REQUIRE(memcmp(Buffer->GetData(), Case.Truth, Case.TruthSize) == 0);
+	}
+}
+
+////////////////////////////////////////////////////////////////////////////////
+void TestCborWriter_Map()
+{
+	// {}
+	{
+		const uint8 Truth[] = { 0xa0 };
+
+		TInlineBuffer<4> Buffer;
+		FCborWriter Writer(Buffer);
+		Writer.OpenMap(0);
+
+		REQUIRE(memcmp(Buffer->GetData(), Truth, sizeof(Truth)) == 0);
+	}
+
+	// {...}
+	{
+		const uint8 Truth[] = { 0xbf, 0xff };
+
+		TInlineBuffer<4> Buffer;
+		FCborWriter Writer(Buffer);
+		Writer.OpenMap();
+		Writer.Close();
+
+		REQUIRE(memcmp(Buffer->GetData(), Truth, sizeof(Truth)) == 0);
+	}
+
+	// { "key":"value", 0:1 }
+	{
+		const uint8 Truth[] = {
+			0xa2,
+			0x63, 0x6b, 0x65, 0x79,
+			0x65, 0x76, 0x61, 0x6c, 0x75, 0x65,
+			0x00, 0x01
+		};
+
+		TInlineBuffer<4> Buffer;
+		FCborWriter Writer(Buffer);
+		Writer.OpenMap(2);
+		Writer.WriteString("key");
+		Writer.WriteString("value", 5);
+		Writer.WriteInteger(0);
+		Writer.WriteInteger(1);
+
+		REQUIRE(memcmp(Buffer->GetData(), Truth, sizeof(Truth)) == 0);
+	}
+}
+
+
+
+// {{{1 test -------------------------------------------------------------------
+
+////////////////////////////////////////////////////////////////////////////////
+void TestCbor()
+{
+	TestCborReader_Integer();
+	TestCborReader_String();
+	TestCborReader_Map();
+
+	TestCborWriter_Integer();
+	TestCborWriter_String();
+	TestCborWriter_Map();
+
+	// [-1]
+	//const uint8 Truth[] = { 0x81, 0x20 };
+
+	// float
+	//const uint8 Truth[] = { 0xfb 0x3f 0xb9 0x99 0x99 0x99 0x99 0x99 0x9a }; // primitive(4591870180066957722)
+}
+
+/* vim: set noexpandtab foldlevel=1 : */
diff --git a/Engine/Source/Programs/UnrealTraceServer/src/Cbor.h b/Engine/Source/Programs/UnrealTraceServer/src/Cbor.h
new file mode 100644
--- /dev/null
+++ b/Engine/Source/Programs/UnrealTraceServer/src/Cbor.h
@@ -0,0 +1,82 @@
+// Copyright Epic Games, Inc. All Rights Reserved.
+
+#pragma once
+
+/*
+ * A note to future visitors; the CBOR implementation below is not complete. Just
+ * enough to satisfy requirements has been implemented
+ */
+
+#include "Foundation.h"
+
+class FInlineBuffer;
+
+////////////////////////////////////////////////////////////////////////////////
+enum class ECborType
+{
+	Error = -1,
+	Eof,
+	Integer,
+	//Boolean,
+	//Float,
+	String,
+	Array,
+	Map,
+	End, // ..of array/map
+	//Null,
+	//Undefined,
+	_Count,
+};
+
+////////////////////////////////////////////////////////////////////////////////
+class FCborContext
+{
+	friend class FCborReader;
+
+public:
+	void			Reset();
+	ECborType		GetType() const;
+	int64			GetLength() const;
+	FStringView		AsString() const;
+	int64			AsInteger() const;
+
+private:
+	void			SetError();
+	const uint8*	NextCursor = nullptr;
+	int64			Param;
+	ECborType		Type;
+};
+
+////////////////////////////////////////////////////////////////////////////////
+class FCborWriter
+{
+public:
+					FCborWriter(FInlineBuffer& InBuffer);
+	void			OpenMap(int32 ItemCount=-1);
+	void			OpenArray(int32 ItemCount=-1);
+	void			Close();
+	void			WriteString(const char* Value, int32 Length=-1);
+	void			WriteInteger(int64 Value);
+
+private:
+	void			WriteParam(uint32 MajorType, uint64 Value);
+	FInlineBuffer&	Buffer;
+};
+
+////////////////////////////////////////////////////////////////////////////////
+class FCborReader
+{
+public:
+					FCborReader(const uint8* Data, uint32 DataSize);
+	bool			ReadNext(FCborContext& Context);
+
+private:
+	int64			ReadParam(const uint8*& Cursor);
+	bool			ReadContainer(FCborContext& Context, const uint8* Cursor);
+	bool			ReadInteger(FCborContext& Context, const uint8* Cursor);
+	bool			ReadString(FCborContext& Context, const uint8* Cursor);
+	const uint8*	Base;
+	const uint8*	DataEnd;
+};
+
+/* vim: set noexpandtab : */
diff --git a/Engine/Source/Programs/UnrealTraceServer/src/CborPayload.h b/Engine/Source/Programs/UnrealTraceServer/src/CborPayload.h
new file mode 100644
--- /dev/null
+++ b/Engine/Source/Programs/UnrealTraceServer/src/CborPayload.h
@@ -0,0 +1,209 @@
+// Copyright Epic Games, Inc. All Rights Reserved.
+
+#pragma once
+
+#include "Cbor.h"
+#include "Foundation.h"
+#include "Utils.h"
+
+////////////////////////////////////////////////////////////////////////////////
+struct FPayload
+{
+	const uint8*	Data;
+	uint32			Size;
+};
+
+
+
+////////////////////////////////////////////////////////////////////////////////
+enum class EStatusCode
+{
+	Unknown				= 0,
+	Success				= 200,
+	BadRequest			= 400,
+	NotFound			= 404,
+	InternalError		= 500,
+};
+
+
+
+////////////////////////////////////////////////////////////////////////////////
+template <int Size=128>
+class TPayloadBuilder
+{
+public:
+							TPayloadBuilder(EStatusCode StatusCode);
+	template <int N>		TPayloadBuilder(const char (&Path)[N]);
+	template <int N> void	AddInteger(const char (&Name)[N], int64 Value);
+	template <int N> void	AddString(const char (&Name)[N], const char* Value, int32 Length=-1);
+	FPayload				Done();
+
+private:
+	TInlineBuffer<Size>		Buffer;
+	FCborWriter				CborWriter = { Buffer };
+};
+
+////////////////////////////////////////////////////////////////////////////////
+template <int Size>
+inline TPayloadBuilder<Size>::TPayloadBuilder(EStatusCode StatusCode)
+{
+	CborWriter.OpenMap();
+	AddInteger("$status", int32(StatusCode));
+}
+
+////////////////////////////////////////////////////////////////////////////////
+template <int Size>
+template <int N>
+inline TPayloadBuilder<Size>::TPayloadBuilder(const char (&Path)[N])
+{
+	CborWriter.OpenMap();
+	AddString("$request", "GET", 3);
+	AddString("$path", Path, N - 1);
+}
+
+////////////////////////////////////////////////////////////////////////////////
+template <int Size>
+template <int N>
+inline void TPayloadBuilder<Size>::AddInteger(const char (&Name)[N], int64 Value)
+{
+	CborWriter.WriteString(Name, N - 1);
+	CborWriter.WriteInteger(Value);
+}
+
+////////////////////////////////////////////////////////////////////////////////
+template <int Size>
+template <int N>
+inline void TPayloadBuilder<Size>::AddString(
+	const char (&Name)[N],
+	const char* Value,
+	int Length)
+{
+	CborWriter.WriteString(Name, N - 1);
+	CborWriter.WriteString(Value, Length);
+}
+
+////////////////////////////////////////////////////////////////////////////////
+template <int Size>
+inline FPayload TPayloadBuilder<Size>::Done()
+{
+	CborWriter.Close();
+	return { Buffer->GetData(), uint32(Buffer->GetSize()) };
+}
+
+
+
+////////////////////////////////////////////////////////////////////////////////
+class FResponse
+{
+public:
+	EStatusCode		GetStatusCode() const;
+	int64			GetInteger(const char* Key, int64 Default) const;
+	template <int N>
+	FStringView		GetString(const char* Key, const char (&Default)[N]) const;
+	const uint8*	GetData() const;
+	uint32			GetSize() const;
+	uint8*			Reserve(uint32 Size);
+
+private:
+	template <typename Type, typename LambdaType>
+	Type			GetValue(const char* Key, Type Default, LambdaType&& Lambda) const;
+	TArray<uint8>	Buffer;
+};
+
+////////////////////////////////////////////////////////////////////////////////
+inline EStatusCode FResponse::GetStatusCode() const
+{
+	int32 Code = int32(GetInteger("$status", 0));
+	return Code ? EStatusCode(Code) : EStatusCode::Unknown;
+}
+
+////////////////////////////////////////////////////////////////////////////////
+inline const uint8* FResponse::GetData() const
+{
+	return Buffer.GetData();
+}
+
+////////////////////////////////////////////////////////////////////////////////
+inline uint32 FResponse::GetSize() const
+{
+	return Buffer.Num();
+}
+
+////////////////////////////////////////////////////////////////////////////////
+inline uint8* FResponse::Reserve(uint32 Size)
+{
+	Buffer.SetNumUninitialized(Size);
+	return Buffer.GetData();
+}
+
+////////////////////////////////////////////////////////////////////////////////
+template <typename Type, typename LambdaType>
+inline Type FResponse::GetValue(const char* Key, Type Default, LambdaType&& Lambda) const
+{
+	FCborReader CborReader(Buffer.GetData(), Buffer.Num());
+	FCborContext Context;
+
+	if (!CborReader.ReadNext(Context) || Context.GetType() != ECborType::Map)
+	{
+		return Default;
+	}
+
+	while (true)
+	{
+		// Read key
+		if (!CborReader.ReadNext(Context) || Context.GetType() != ECborType::String)
+		{
+			return Default;
+		}
+
+		// Check the key
+		FStringView String = Context.AsString();
+		bool bIsTarget = (String.Compare(Key) == 0);
+
+		// Read value
+		if (!CborReader.ReadNext(Context))
+		{
+			return Default;
+		}
+
+		if (bIsTarget)
+		{
+			return Lambda(Context);
+		}
+	}
+}
+
+////////////////////////////////////////////////////////////////////////////////
+inline int64 FResponse::GetInteger(const char* Key, int64 Default) const
+{
+	return GetValue(
+		Key,
+		Default,
+		[this] (const FCborContext& Context)
+		{
+			return Context.AsInteger();
+		}
+	);
+}
+
+////////////////////////////////////////////////////////////////////////////////
+template <int N>
+inline FStringView FResponse::GetString(const char* Key, const char (&Default)[N]) const
+{
+	FStringView DefaultView(Default, N - 1);
+	return GetValue(
+		Key,
+		DefaultView,
+		[this, DefaultView] (const FCborContext& Context)
+		{
+			if (Context.GetType() == ECborType::String)
+			{
+				return Context.AsString();
+			}
+
+			return DefaultView;
+		}
+	);
+}
+
+/* vim: set noexpandtab : */
diff --git a/Engine/Source/Programs/UnrealTraceServer/src/Foundation.h b/Engine/Source/Programs/UnrealTraceServer/src/Foundation.h
new file mode 100644
--- /dev/null
+++ b/Engine/Source/Programs/UnrealTraceServer/src/Foundation.h
@@ -0,0 +1,83 @@
+// Copyright Epic Games, Inc. All Rights Reserved.
+
+#pragma once
+
+////////////////////////////////////////////////////////////////////////////////
+#define MoveTemp(x) std::move(x)
+#define check(x)
+
+////////////////////////////////////////////////////////////////////////////////
+template <typename Type>
+struct TArray
+	: protected std::vector<Type>
+{
+	using Super = std::vector<Type>;
+	using Super::Super;
+	using Super::operator [];
+	using Super::begin;
+	using Super::end;
+
+	void		Add(const Type& Value)				{ return Super::push_back(Value); }
+	void		SetNum(int32 Num)					{ return Super::resize(Num); }
+	int32		Num() const							{ return int32(Super::size()); }
+	Type*		GetData()							{ return &(this->operator [] (0)); }
+	const Type* GetData() const						{ return &(this->operator [] (0)); }
+	void		Empty()								{ return Super::clear(); }
+	void		SetNumUninitialized(int Num)		{ return Super::resize(Num); }
+	void		Append(const void* Data, int Num)	{ Super::insert(Super::end(), (const Type*)Data, ((const Type*)Data) + Num); }
+};
+
+////////////////////////////////////////////////////////////////////////////////
+struct FStringView
+	: protected std::basic_string_view<char>
+{
+	typedef std::basic_string_view<char> Super;
+	using Super::Super;
+	using Super::operator [];
+
+	const char* GetData() const			{ return Super::data(); }
+	uint32 Len() const					{ return uint32(Super::size()); }
+	int Compare(const char* Rhs) const	{ return Super::compare(Rhs); }
+};
+
+////////////////////////////////////////////////////////////////////////////////
+struct FString
+	: protected std::string
+{
+	typedef std::string Super;
+	using Super::Super;
+	using Super::operator +=;
+
+	const char* operator * () const				{ return Super::c_str(); }
+	void operator += (const FStringView& Rhs)	{ Super::operator += ((FStringView::Super&)(Rhs)); }
+};
+
+
+
+////////////////////////////////////////////////////////////////////////////////
+#if TS_USING(TS_PLATFORM_WINDOWS)
+struct FWinApiStr
+{
+	FWinApiStr(const char* Utf8)
+	{
+		int32 BufferSize = MultiByteToWideChar(CP_UTF8, 0, Utf8, -1, nullptr, 0);
+		Buffer = new wchar_t[BufferSize];
+		MultiByteToWideChar(CP_UTF8, 0, Utf8, -1, Buffer, BufferSize);
+	}
+
+	~FWinApiStr()
+	{
+		delete[] Buffer;
+	}
+
+	operator LPCWSTR () const
+	{
+		return Buffer;
+	}
+
+private:
+	wchar_t* Buffer = nullptr;
+};
+#endif // TS_PLATFORM_WINDOWS
+
+/* vim: set noexpandtab : */
diff --git a/Engine/Source/Programs/UnrealTraceServer/src/Main.cpp b/Engine/Source/Programs/UnrealTraceServer/src/Main.cpp
new file mode 100644
--- /dev/null
+++ b/Engine/Source/Programs/UnrealTraceServer/src/Main.cpp
@@ -0,0 +1,624 @@
+// Copyright Epic Games, Inc. All Rights Reserved.
+
+#include "Pch.h"
+#include "StoreService.h"
+#include "Utils.h"
+#include "Version.h"
+
+#if TS_USING(TS_BUILD_DEBUG)
+#	include <thread>
+#endif
+
+#if TS_USING(TS_PLATFORM_LINUX) || TS_USING(TS_PLATFORM_MAC)
+#	include <pwd.h>
+#	include <semaphore.h>
+#	include <signal.h>
+#	include <sys/mman.h>
+#	include <unistd.h>
+#endif
+
+#include <immintrin.h>
+
+// {{{1 misc -------------------------------------------------------------------
+
+////////////////////////////////////////////////////////////////////////////////
+class FMmapScope
+{
+public:
+								FMmapScope(void* InPtr, uint32 InLength=0);
+								~FMmapScope();
+	template <typename T> T*	As() const;
+
+private:
+	void*						Ptr;
+	int32						Length;
+
+private:
+								FMmapScope() = delete;
+								FMmapScope(const FMmapScope&) = delete;
+								FMmapScope(const FMmapScope&&) = delete;
+	FMmapScope&					operator = (const FMmapScope&) = delete;
+	FMmapScope&					operator = (const FMmapScope&&) = delete;
+};
+
+////////////////////////////////////////////////////////////////////////////////
+FMmapScope::FMmapScope(void* InPtr, uint32 InLength)
+: Ptr(InPtr)
+, Length(InLength)
+{
+}
+
+////////////////////////////////////////////////////////////////////////////////
+FMmapScope::~FMmapScope()
+{
+#if TS_USING(TS_PLATFORM_WINDOWS)
+	UnmapViewOfFile(Ptr);
+#elif TS_USING(TS_PLATFORM_LINUX) || TS_USING(TS_PLATFORM_MAC)
+	munmap(Ptr, Length);
+#endif
+}
+
+////////////////////////////////////////////////////////////////////////////////
+template <typename T>
+T* FMmapScope::As() const
+{
+	return (T*)Ptr;
+}
+
+
+
+////////////////////////////////////////////////////////////////////////////////
+template <typename T>
+struct TOnScopeExit
+{
+		TOnScopeExit()	= default;
+		~TOnScopeExit()	{ (*Lambda)(); }
+	T*	Lambda;
+};
+
+#define OnScopeExit(x) \
+	auto TS_CONCAT(OnScopeExitFunc, __LINE__) = x; \
+	TOnScopeExit<decltype(TS_CONCAT(OnScopeExitFunc, __LINE__))> TS_CONCAT(OnScopeExitInstance, __LINE__); \
+	do { TS_CONCAT(OnScopeExitInstance, __LINE__).Lambda = &TS_CONCAT(OnScopeExitFunc, __LINE__); } while (0)
+
+
+
+
+// {{{1 store ------------------------------------------------------------------
+
+////////////////////////////////////////////////////////////////////////////////
+static FStoreService* StartStore(const char* StoreDir)
+{
+	FStoreService::FDesc Desc;
+	Desc.StoreDir = StoreDir;
+	Desc.StorePort = 1989;
+	Desc.RecorderPort = 1980;
+	return FStoreService::Create(Desc);
+}
+
+
+
+// {{{1 instance-info ----------------------------------------------------------
+
+////////////////////////////////////////////////////////////////////////////////
+struct FInstanceInfo
+{
+public:
+	static const uint32	CurrentVersion =
+#if TS_USING(TS_BUILD_DEBUG)
+		0x8000'0000 |
+#endif
+		((TS_VERSION_PROTOCOL & 0xffff) << 16) | (TS_VERSION_MINOR & 0xffff);
+
+	void				Set();
+	void				WaitForReady() const;
+	bool				IsOlder() const;
+	std::atomic<uint32> Published;
+	uint32				Version;
+	uint32				Pid;
+};
+
+////////////////////////////////////////////////////////////////////////////////
+void FInstanceInfo::Set()
+{
+	Version = 1;
+#if TS_USING(TS_PLATFORM_WINDOWS)
+	Pid = GetCurrentProcessId();
+#elif TS_USING(TS_PLATFORM_LINUX) || TS_USING(TS_PLATFORM_MAC)
+	Pid = getpid();
+#endif
+	Published.fetch_add(1, std::memory_order_release);
+}
+
+////////////////////////////////////////////////////////////////////////////////
+void FInstanceInfo::WaitForReady() const
+{
+	// Spin until this instance info is published (by another process)
+	for (;; _mm_pause())
+	{
+		if (Published.load(std::memory_order_acquire))
+		{
+			break;
+		}
+	}
+}
+
+////////////////////////////////////////////////////////////////////////////////
+bool FInstanceInfo::IsOlder() const
+{
+	// Decide which is older; this compiled code or the instance we have a
+	// pointer to.
+	bool bIsOlder = false;
+	bIsOlder |= (Version < FInstanceInfo::CurrentVersion);
+	return bIsOlder;
+}
+
+
+
+// {{{1 windows ----------------------------------------------------------------
+
+#if TS_USING(TS_PLATFORM_WINDOWS)
+////////////////////////////////////////////////////////////////////////////////
+class FWinHandle
+{
+public:
+	FWinHandle(HANDLE InHandle)
+	: Handle(InHandle)
+	{
+		if (Handle == INVALID_HANDLE_VALUE)
+		{
+			Handle = nullptr;
+		}
+	}
+
+				~FWinHandle()				{ if (Handle) CloseHandle(Handle); }
+				operator HANDLE () const	{ return Handle; }
+	bool		IsValid() const				{ return Handle != nullptr; }
+
+private:
+	HANDLE		Handle;
+	FWinHandle&	operator = (const FWinHandle&) = delete;
+				FWinHandle(const FWinHandle&) = delete;
+				FWinHandle(const FWinHandle&&) = delete;
+};
+
+////////////////////////////////////////////////////////////////////////////////
+static const wchar_t*	GIpcName					= L"Local\\UnrealTraceInstance";
+static const int32		GIpcSize					= 4 << 10;
+static const wchar_t*	GQuitEventName				= L"Local\\UnrealTraceEvent";
+static const wchar_t*	GBegunEventName				= L"Local\\UnrealTraceEventBegun";
+static int				MainDaemon(int, char**);
+
+////////////////////////////////////////////////////////////////////////////////
+static int MainFork(int ArgC, char** ArgV)
+{
+	// Check for an existing instance that is already running.
+	FWinHandle IpcHandle = OpenFileMappingW(FILE_MAP_ALL_ACCESS, false, GIpcName);
+	if (IpcHandle.IsValid())
+	{
+		void* IpcPtr = MapViewOfFile(IpcHandle, FILE_MAP_ALL_ACCESS, 0, 0, GIpcSize);
+		FMmapScope MmapScope(IpcPtr);
+
+		const auto* InstanceInfo = MmapScope.As<const FInstanceInfo>();
+		InstanceInfo->WaitForReady();
+		if (!InstanceInfo->IsOlder())
+		{
+			return 0;
+		}
+
+		// Signal to the existing instance to shutdown or forcefully do it if it
+		// does not respond in time.
+		FWinHandle QuitEvent = CreateEventW(nullptr, TRUE, FALSE, GQuitEventName);
+		if (GetLastError() != ERROR_ALREADY_EXISTS)
+		{
+			// Assume someone else has closed the instance already.
+			return 0;
+		}
+
+		DWORD OpenProcessFlags = PROCESS_TERMINATE | SYNCHRONIZE;
+		FWinHandle ProcHandle = OpenProcess(OpenProcessFlags, FALSE, InstanceInfo->Pid);
+		if (!ProcHandle.IsValid())
+		{
+			return 5;
+		}
+
+		SetEvent(QuitEvent);
+
+		if (WaitForSingleObject(ProcHandle, 5000) == WAIT_TIMEOUT)
+		{
+			TerminateProcess(ProcHandle, 10);
+		}
+	}
+
+	// Get this binary's path
+	wchar_t BinPath[MAX_PATH];
+	uint32 BinPathLen = GetModuleFileNameW(nullptr, BinPath, TS_ARRAY_COUNT(BinPath));
+	if (GetLastError() == ERROR_INSUFFICIENT_BUFFER)
+	{
+		// This should never really happen...
+		return 1;
+	}
+
+	// Calculate where to store the binaries.
+	std::filesystem::path DestPath;
+	{
+		wchar_t Buffer[MAX_PATH];
+		SHGetFolderPathW(nullptr, CSIDL_LOCAL_APPDATA, nullptr, SHGFP_TYPE_CURRENT, Buffer);
+		DestPath = Buffer;
+	}
+	DestPath /= "UnrealEngine/Common/UnrealTrace";
+	{
+		wchar_t Buffer[64];
+		_snwprintf(Buffer, TS_ARRAY_COUNT(Buffer), L"Bin/%08x/UnrealTraceServer.exe", FInstanceInfo::CurrentVersion);
+		DestPath /= Buffer;
+	}
+
+	// Copy the binary out to a location where it doesn't matter if the file
+	// gets locked by the OS.
+	if (!std::filesystem::is_regular_file(DestPath))
+	{
+		std::filesystem::create_directories(DestPath.parent_path());
+
+		// Tag the destination with our PID and copy
+		DWORD OurPid = GetCurrentProcessId();
+		wchar_t Buffer[16];
+		_snwprintf(Buffer, TS_ARRAY_COUNT(Buffer), L"_%08x", OurPid);
+		std::filesystem::path TempPath = DestPath;
+		TempPath += Buffer;
+		if (!std::filesystem::copy_file(BinPath, TempPath))
+		{
+			return 3;
+		}
+
+		// Move the file into place. If this fails because the file exists then
+		// another instance has beaten us to the punch.
+		std::error_code ErrorCode;
+		std::filesystem::rename(TempPath, DestPath, ErrorCode);
+		if (ErrorCode)
+		{
+			bool RaceLost = (ErrorCode == std::errc::file_exists);
+			return RaceLost ? 0 : 4;
+		}
+	}
+
+	// Launch a new instance as a daemon and wait until we know it has started
+	FWinHandle BegunEvent = CreateEventW(nullptr, TRUE, FALSE, GBegunEventName);
+	if (GetLastError() == ERROR_ALREADY_EXISTS)
+	{
+		return 7;
+	}
+
+	// For debugging ease and consistency we will daemonize in this process
+	// instead of spawning a second one.
+#if TS_USING(TS_BUILD_DEBUG)
+	std::thread DaemonThread([] () { MainDaemon(0, nullptr); });
+#else
+	uint32 Flags = CREATE_BREAKAWAY_FROM_JOB;
+	STARTUPINFOW StartupInfo = { sizeof(STARTUPINFOW) };
+	PROCESS_INFORMATION ProcessInfo = {};
+	BOOL bOk = CreateProcessW(DestPath.c_str(), LPWSTR(L"UnrealTraceServer.exe daemon"), nullptr,
+		nullptr, FALSE, Flags, nullptr, nullptr, &StartupInfo, &ProcessInfo);
+
+	if (bOk == FALSE)
+	{
+		return 2;
+	}
+#endif // TS_BUILD_DEBUG
+
+	int Ret = 0;
+	if (WaitForSingleObject(BegunEvent, 5000) == WAIT_TIMEOUT)
+	{
+		Ret = 8;
+	}
+
+#if TS_USING(TS_BUILD_DEBUG)
+	DaemonThread.join();
+#else
+	CloseHandle(ProcessInfo.hProcess);
+	CloseHandle(ProcessInfo.hThread);
+#endif
+
+	return Ret;
+}
+
+////////////////////////////////////////////////////////////////////////////////
+static int MainDaemon(int ArgC, char** ArgV)
+{
+	// Create a piece of shared memory so all store instances can communicate.
+	FWinHandle IpcHandle = CreateFileMappingW(INVALID_HANDLE_VALUE, nullptr,
+		PAGE_READWRITE, 0, GIpcSize, GIpcName);
+	if (!IpcHandle.IsValid())
+	{
+		return 1;
+	}
+
+	// Create a named event so others can tell us to quit.
+	FWinHandle QuitEvent = CreateEventW(nullptr, TRUE, FALSE, GQuitEventName);
+	if (GetLastError() == ERROR_ALREADY_EXISTS)
+	{
+		// This really should not happen. It is expected that only one process
+		// will get this far (gated by the shared-memory object creation).
+		return 2;
+	}
+
+	// Fill out the Ipc details and publish.
+	void* IpcPtr = MapViewOfFile(IpcHandle, FILE_MAP_ALL_ACCESS, 0, 0, GIpcSize);
+	{
+		FMmapScope MmapScope(IpcPtr);
+		auto* InstanceInfo = MmapScope.As<FInstanceInfo>();
+		InstanceInfo->Set();
+	}
+
+	// Fire up the store
+	FStoreService* StoreService;
+	{
+		std::filesystem::path StoreDir;
+		wchar_t Buffer[MAX_PATH];
+		SHGetFolderPathW(nullptr, CSIDL_LOCAL_APPDATA, nullptr, SHGFP_TYPE_CURRENT, Buffer);
+		StoreDir = Buffer;
+		StoreDir /= "UnrealEngine/Common/UnrealTrace/Store";
+
+		std::u8string StoreDirUtf8 = StoreDir.u8string();
+		StoreService = StartStore((const char*)(StoreDirUtf8.c_str()));
+	}
+
+	// Let every one know we've started.
+	{
+		FWinHandle BegunEvent = CreateEventW(nullptr, TRUE, FALSE, GBegunEventName);
+		if (BegunEvent.IsValid())
+		{
+			SetEvent(BegunEvent);
+		}
+	}
+
+	// Wait to be told to resign.
+	WaitForSingleObject(QuitEvent, INFINITE);
+
+	// Clean up. We are done here.
+	delete StoreService;
+	return 0;
+}
+
+////////////////////////////////////////////////////////////////////////////////
+int WinMain(HINSTANCE, HINSTANCE, LPSTR, int)
+{
+	// Get command line arguments and convert to UTF8
+	const wchar_t* CommandLine = GetCommandLineW();
+	int ArgC;
+	wchar_t** WideArgV = CommandLineToArgvW(CommandLine, &ArgC);
+
+	TInlineBuffer<320> ArgBuffer;
+	char** ArgV = (char**)(ArgBuffer->Append(ArgC * sizeof(char*)));
+	for (int i = 0; i < ArgC; ++i)
+	{
+		const wchar_t* WideArg = WideArgV[i];
+		int32 ArgSize = WideCharToMultiByte(CP_UTF8, 0, WideArg, -1, nullptr, 0, nullptr, nullptr);
+
+		char* Arg = (char*)(ArgBuffer->Append(ArgSize));
+		WideCharToMultiByte(CP_UTF8, 0, WideArg, -1, Arg, ArgSize, nullptr, nullptr);
+
+		ArgV[i] = Arg;
+	}
+
+	// The proper entry point
+	extern int main(int, char**);
+	int Ret = main(ArgC, ArgV);
+
+	// Clean up
+	LocalFree(ArgV);
+	return Ret;
+}
+
+#endif // TS_PLATFORM_WINDOWS
+
+
+
+// {{{1 linux/mac --------------------------------------------------------------
+
+#if TS_USING(TS_PLATFORM_LINUX) || TS_USING(TS_PLATFORM_MAC)
+
+////////////////////////////////////////////////////////////////////////////////
+static const char*		GShmName					= "/UnrealTraceShm";
+static const int32		GShmSize					= 4 << 10;
+static const char*		GBegunSemName				= "/UnrealTraceSemBegun";
+static const char*		GQuitSemName				= "/UnrealTraceSemQuit";
+static int				MainDaemon(int, char**);
+
+////////////////////////////////////////////////////////////////////////////////
+static int MainFork(int ArgC, char** ArgV)
+{
+	// Check for an existing instance that is already running.
+	int ShmHandle = shm_open(GShmName, O_RDONLY, S_IROTH|S_IWOTH);
+	if (ShmHandle >= 0)
+	{
+		void* ShmPtr = mmap(nullptr, GShmSize, PROT_READ, MAP_SHARED, ShmHandle, 0);
+		FMmapScope MmapScope(ShmPtr, GShmSize);
+		close(ShmHandle);
+
+		const auto* InstanceInfo = MmapScope.As<const FInstanceInfo>();
+		InstanceInfo->WaitForReady();
+		if (!InstanceInfo->IsOlder())
+		{
+			return 0;
+		}
+
+		// Signal to the existing instance to shutdown.
+		sem_t* QuitSem = sem_open(GQuitSemName, O_RDWR, 0666, 0);
+		if (QuitSem == SEM_FAILED)
+		{
+			// Assume someone else has closed the instance already.
+			return 0;
+		}
+		sem_post(QuitSem);
+		sem_close(QuitSem);
+
+		// Wait for the process to end. If it takes too long, kill it.
+		for (uint32 i = 0; ; i += 1000)
+		{
+			if (kill(InstanceInfo->Pid, 0) < 0)
+			{
+				break;
+			}
+
+			if (i >= 5000) // "5000" for grap-ability
+			{
+				kill(InstanceInfo->Pid, SIGTERM);
+				break;
+			}
+		}
+	}
+
+	// Launch a new instance as a daemon and wait until we know it has started
+	sem_t* BegunSem = sem_open(GBegunSemName, O_RDONLY|O_CREAT|O_EXCL, 0644, 0);
+	if (BegunSem == SEM_FAILED)
+	{
+		return (errno == EEXIST) ? 7 : 11;
+	}
+	OnScopeExit([=] () {
+		sem_unlink(GBegunSemName);
+		sem_close(BegunSem);
+	});
+
+	// For debugging ease and consistency we will daemonize in this process
+	// instead of spawning a second one.
+#if TS_USING(TS_BUILD_DEBUG)
+	std::thread DaemonThread([] () { MainDaemon(0, nullptr); });
+#else
+	pid_t DaemonPid = fork();
+	if (DaemonPid < 0)
+	{
+		return 3;
+	}
+	else if (DaemonPid == 0)
+	{
+		return MainDaemon(0, nullptr);
+	}
+#endif // TS_BUILD_DEBUG
+
+	int Ret = 0;
+	timespec Timeout;
+	clock_gettime(CLOCK_REALTIME, &Timeout);
+	Timeout.tv_sec += (5000 / 1000); // expressed in milliseconds for grep-ability
+	if (sem_timedwait(BegunSem, &Timeout) < 0)
+	{
+		Ret = 8;
+	}
+
+#if TS_USING(TS_BUILD_DEBUG)
+	DaemonThread.join();
+#endif
+
+	return Ret;
+}
+
+////////////////////////////////////////////////////////////////////////////////
+static int MainDaemon(int ArgC, char** ArgV)
+{
+	// Create a piece of shared memory so all store instances can communicate.
+	int ShmHandle = shm_open(GShmName, O_RDWR|O_CREAT|O_EXCL, S_IROTH|S_IWOTH);
+	if (ShmHandle < 0)
+	{
+		return 1;
+	}
+	OnScopeExit([=] () {
+		shm_unlink(GShmName);
+		close(ShmHandle);
+	});
+
+	// Create a named semaphore so others can tell us to quit.
+	sem_t* QuitSem = sem_open(GQuitSemName, O_RDWR|O_CREAT|O_EXCL, 0666, 0);
+	if (QuitSem == SEM_FAILED)
+	{
+		// This really should not happen. It is expected that only one process
+		// will get this far (gated by the shared-memory object creation).
+		return 2;
+	}
+	OnScopeExit([=] () {
+		sem_unlink(GQuitSemName);
+		sem_close(QuitSem);
+	});
+
+	// Fill out the Ipc details and publish.
+	if (ftruncate(ShmHandle, GShmSize) != 0)
+	{
+		return 3;
+	}
+
+	int32 ProtFlags = PROT_READ|PROT_WRITE;
+	void* ShmPtr = mmap(nullptr, GShmSize, ProtFlags, MAP_SHARED, ShmHandle, 0);
+	{
+		FMmapScope MmapScope(ShmPtr, GShmSize);
+		auto* InstanceInfo = MmapScope.As<FInstanceInfo>();
+		InstanceInfo->Set();
+	}
+
+	// Fire up the store
+	FStoreService* StoreService;
+	{
+		int UserId = getuid();
+		const passwd* Passwd = getpwuid(UserId);
+
+		std::filesystem::path StoreDir = Passwd->pw_dir;
+		StoreDir /= "UnrealEngine/UnrealTrace/Store";
+
+		StoreService = StartStore(StoreDir.c_str());
+	}
+
+	// Let every one know we've started.
+	sem_t* BegunSem = sem_open(GBegunSemName, O_RDWR, 0644, 0);
+	if (BegunSem != SEM_FAILED)
+	{
+		sem_post(BegunSem);
+		sem_close(BegunSem);
+	}
+
+	// Wait to be told to resign.
+	sem_wait(QuitSem);
+
+	// Clean up. We are done here.
+	delete StoreService;
+	return 0;
+}
+
+#endif // TS_PLATFORM_LINUX/MAC
+
+
+
+// {{{1 main -------------------------------------------------------------------
+
+////////////////////////////////////////////////////////////////////////////////
+int MainTest(int ArgC, char** ArgV)
+{
+	extern void TestCbor();
+	TestCbor();
+	return 0;
+}
+
+////////////////////////////////////////////////////////////////////////////////
+int main(int ArgC, char** ArgV)
+{
+	struct {
+		const char*	Verb;
+		int			(*Entry)(int, char**);
+	} Dispatches[] = {
+		"fork",		MainFork,
+		"daemon",	MainDaemon,
+		"test",		MainTest,
+		//"kill",		MainKill,
+	};
+
+	if (ArgC > 1)
+	{
+		for (const auto& Dispatch : Dispatches)
+		{
+			if (strcmp(ArgV[1], Dispatch.Verb) == 0)
+			{
+				return (Dispatch.Entry)(ArgC - 1, ArgV + 1);
+			}
+		}
+	}
+
+	return MainFork(ArgC, ArgV);
+}
+
+/* vim: set noexpandtab foldlevel=1 : */
diff --git a/Engine/Source/Programs/UnrealTraceServer/src/Pch.h b/Engine/Source/Programs/UnrealTraceServer/src/Pch.h
new file mode 100644
--- /dev/null
+++ b/Engine/Source/Programs/UnrealTraceServer/src/Pch.h
@@ -0,0 +1,90 @@
+// Copyright Epic Games, Inc. All Rights Reserved.
+
+#pragma once
+
+// {{{1 compiler ---------------------------------------------------------------
+#if defined(_MSC_VER)
+#	pragma warning(disable: 4200) // zero-sized arrays are non-standard
+#endif
+
+
+
+// {{{1 build ------------------------------------------------------------------
+
+#define TS_ON					+
+#define TS_OFF					-
+#define TS_USING(x)				((1 x 1) == 2)
+
+#define TS_PLATFORM_WINDOWS		TS_OFF
+#define TS_PLATFORM_LINUX		TS_OFF
+#define TS_PLATFORM_MAC			TS_OFF
+
+#if defined(_MSC_VER)
+#	undef  TS_PLATFORM_WINDOWS
+#	define TS_PLATFORM_WINDOWS	TS_ON
+#endif
+
+#if defined(__linux__)
+#	undef  TS_PLATFORM_LINUX
+#	define TS_PLATFORM_LINUX	TS_ON
+#endif
+
+#if defined(__APPLE__)
+#	undef  TS_PLATFORM_MAC
+#	define TS_PLATFORM_MAC		TS_ON
+#endif
+
+#define TS_ARRAY_COUNT(x)		(sizeof(x) / sizeof((x)[0]))
+
+#if !defined(TS_BUILD_DEBUG)
+#	define TS_BUILD_DEBUG		TS_OFF
+#endif
+
+
+
+// {{{1 base -------------------------------------------------------------------
+
+#if TS_USING(TS_PLATFORM_WINDOWS)
+#	define _HAS_EXCEPTIONS			0
+#	define _CRT_SECURE_NO_WARNINGS	0
+#endif
+
+#include <algorithm>
+#include <atomic>
+#include <chrono>
+#include <cstdint>
+#include <ctime>
+#include <filesystem>
+#include <string>
+#include <string_view>
+#include <vector>
+
+#if TS_USING(TS_PLATFORM_WINDOWS)
+#	define _WIN32_WINNT			0x0a00
+#	define WIN32_LEAN_AND_MEAN
+#	define NOMINMAX
+#	include <Windows.h>
+#	include <Shlobj.h>
+#	include <shellapi.h>
+#	pragma comment(lib, "Shell32.lib")
+#endif
+
+#include "Asio.h"
+
+using int8		= int8_t;
+using int16		= int16_t;
+using int32		= int32_t;
+using int64		= int64_t;
+using uint8		= uint8_t;
+using uint16	= uint16_t;
+using uint32	= uint32_t;
+using uint64	= uint64_t;
+
+
+
+// {{{1 util -------------------------------------------------------------------
+
+#define TS_CONCAT_(a, b)		a##b
+#define TS_CONCAT(a, b)			TS_CONCAT_(a, b)
+
+/* vim: set noexpandtab foldlevel=1 : */
diff --git a/Engine/Source/Programs/UnrealTraceServer/src/Recorder.cpp b/Engine/Source/Programs/UnrealTraceServer/src/Recorder.cpp
new file mode 100644
--- /dev/null
+++ b/Engine/Source/Programs/UnrealTraceServer/src/Recorder.cpp
@@ -0,0 +1,336 @@
+// Copyright Epic Games, Inc. All Rights Reserved.
+
+#include "Pch.h"
+#include "AsioIoable.h"
+#include "AsioSocket.h"
+#include "Recorder.h"
+#include "Store.h"
+
+#if TS_USING(TS_PLATFORM_WINDOWS)
+#	include <Mstcpip.h>
+#endif
+
+////////////////////////////////////////////////////////////////////////////////
+constexpr uint32 FourCc(int A, int B, int C, int D)
+{
+    return ((A << 24) | (B << 16) | (C << 8) | D);
+}
+
+
+
+////////////////////////////////////////////////////////////////////////////////
+class FRecorderRelay
+	: public FAsioIoSink
+{
+public:
+						FRecorderRelay(asio::ip::tcp::socket& Socket, FAsioWriteable* InOutput);
+	virtual				~FRecorderRelay();
+	bool				IsOpen();
+	void				Close();
+	uint32				GetIpAddress() const;
+	uint32				GetControlPort() const;
+
+private:
+	virtual void		OnIoComplete(uint32 Id, int32 Size) override;
+	bool				ReadMetadata(int32 Size);
+	static const uint32	BufferSize = 64 * 1024;
+	enum				{ OpStart, OpSocketReadMetadata, OpSocketRead, OpFileWrite };
+	FAsioSocket			Input;
+	FAsioWriteable*		Output;
+	uint32				ActiveReadOp = OpSocketReadMetadata;
+	uint16				ControlPort = 0;
+	uint8				Buffer[BufferSize];
+};
+
+////////////////////////////////////////////////////////////////////////////////
+FRecorderRelay::FRecorderRelay(asio::ip::tcp::socket& Socket, FAsioWriteable* InOutput)
+: Input(Socket)
+, Output(InOutput)
+{
+	OnIoComplete(OpStart, 0);
+}
+
+////////////////////////////////////////////////////////////////////////////////
+FRecorderRelay::~FRecorderRelay()
+{
+	check(!Input.IsOpen());
+	check(!Output->IsOpen());
+	delete Output;
+}
+
+////////////////////////////////////////////////////////////////////////////////
+bool FRecorderRelay::IsOpen()
+{
+	return Input.IsOpen();
+}
+
+////////////////////////////////////////////////////////////////////////////////
+void FRecorderRelay::Close()
+{
+	Input.Close();
+	Output->Close();
+}
+
+////////////////////////////////////////////////////////////////////////////////
+uint32 FRecorderRelay::GetIpAddress() const
+{
+	return Input.GetRemoteAddress();
+}
+
+////////////////////////////////////////////////////////////////////////////////
+uint32 FRecorderRelay::GetControlPort() const
+{
+	return ControlPort;
+}
+
+////////////////////////////////////////////////////////////////////////////////
+bool FRecorderRelay::ReadMetadata(int32 Size)
+{
+	const uint8* Cursor = Buffer;
+	auto Read = [&] (int32 SizeToRead)
+	{
+		const uint8* Ptr = Cursor;
+		Cursor += SizeToRead;
+		Size -= SizeToRead;
+		return Ptr;
+	};
+
+	// Stream header
+	uint32 Magic;
+	if (Size < sizeof(Magic))
+	{
+		return true;
+	}
+
+	Magic = *(const uint32*)(Read(sizeof(Magic)));
+	switch (Magic)
+	{
+	/* trace with metadata data */
+	case FourCc('T', 'R', 'C', '2'):
+		break;
+
+	/* valid, but to old or wrong endian for us */
+	case FourCc('T', 'R', 'C', 'E'):
+	case FourCc('E', 'C', 'R', 'T'):
+	case FourCc('2', 'C', 'R', 'T'):
+		return true;
+
+	/* unexpected magic */
+	default:
+		return false;
+	}
+
+	// MetadataSize field
+	if (Size < 2)
+	{
+		return true;
+	}
+	Size = *(const uint16*)(Read(2));
+
+	// MetadataFields
+	while (Size >= 2)
+	{
+		struct {
+			uint8	Size;
+			uint8	Id;
+		} MetadataField;
+		MetadataField = *(const decltype(MetadataField)*)(Read(sizeof(MetadataField)));
+
+		if (Size < MetadataField.Size)
+		{
+			break;
+		}
+
+		if (MetadataField.Id == 0) /* ControlPortFieldId */
+		{
+			ControlPort = *(const uint16*)Cursor;
+		}
+
+		Size -= MetadataField.Size;
+	}
+
+	return true;
+}
+
+////////////////////////////////////////////////////////////////////////////////
+void FRecorderRelay::OnIoComplete(uint32 Id, int32 Size)
+{
+	if (Size < 0)
+	{
+		Close();
+		return;
+	}
+
+	switch (Id)
+	{
+	case OpSocketReadMetadata:
+		ActiveReadOp = OpSocketRead;
+		if (!ReadMetadata(Size))
+		{
+			Close();
+			return;
+		}
+		/* fallthrough */
+
+	case OpSocketRead:
+		Output->Write(Buffer, Size, this, OpFileWrite);
+		break;
+
+	case OpStart:
+	case OpFileWrite:
+		Input.ReadSome(Buffer, BufferSize, this, ActiveReadOp);
+		break;
+	}
+}
+
+
+
+////////////////////////////////////////////////////////////////////////////////
+uint32 FRecorder::FSession::GetId() const
+{
+	return Id;
+}
+
+////////////////////////////////////////////////////////////////////////////////
+uint32 FRecorder::FSession::GetTraceId() const
+{
+	return TraceId;
+}
+
+////////////////////////////////////////////////////////////////////////////////
+uint32 FRecorder::FSession::GetIpAddress() const
+{
+	return Relay->GetIpAddress();
+}
+
+////////////////////////////////////////////////////////////////////////////////
+uint32 FRecorder::FSession::GetControlPort() const
+{
+	return Relay->GetControlPort();
+}
+
+
+
+////////////////////////////////////////////////////////////////////////////////
+FRecorder::FRecorder(asio::io_context& IoContext, FStore& InStore)
+: FAsioTcpServer(IoContext)
+, FAsioTickable(IoContext)
+, Store(InStore)
+{
+	StartTick(500);
+}
+
+////////////////////////////////////////////////////////////////////////////////
+FRecorder::~FRecorder()
+{
+	check(!FAsioTickable::IsActive());
+	check(!FAsioTcpServer::IsOpen());
+
+	for (FSession& Session : Sessions)
+	{
+		delete Session.Relay;
+	}
+}
+
+////////////////////////////////////////////////////////////////////////////////
+void FRecorder::Close()
+{
+	FAsioTickable::StopTick();
+	FAsioTcpServer::Close();
+
+	for (FSession& Session : Sessions)
+	{
+		Session.Relay->Close();
+	}
+}
+
+////////////////////////////////////////////////////////////////////////////////
+uint32 FRecorder::GetSessionCount() const
+{
+	return uint32(Sessions.Num());
+}
+
+////////////////////////////////////////////////////////////////////////////////
+const FRecorder::FSession* FRecorder::GetSessionInfo(uint32 Index) const
+{
+	if (Index >= uint32(Sessions.Num()))
+	{
+		return nullptr;
+	}
+
+	return Sessions.GetData() + Index;
+}
+
+////////////////////////////////////////////////////////////////////////////////
+bool FRecorder::OnAccept(asio::ip::tcp::socket& Socket)
+{
+	FStore::FNewTrace Trace = Store.CreateTrace();
+	if (Trace.Writeable == nullptr)
+	{
+		return true;
+	}
+
+#if TS_USING(TS_PLATFORM_WINDOWS)
+	// Trace data is a stream and communication is one way. It is implemented
+	// this way to share code between sending trace data over the wire and writing
+	// it to a file. Because there's no ping/pong we can end up with a half-open
+	// TCP connection if the other end doesn't close its socket. So we'll enable
+	// keep-alive on the socket and set a short timeout (default is 2hrs).
+	tcp_keepalive KeepAlive =
+	{
+		1,		// on
+		15000,	// timeout_ms
+		2000,	// interval_ms
+	};
+
+	DWORD BytesReturned;
+	WSAIoctl(
+		Socket.native_handle(),
+		SIO_KEEPALIVE_VALS,
+		&KeepAlive, sizeof(KeepAlive),
+		nullptr, 0,
+		&BytesReturned,
+		nullptr,
+		nullptr
+	);
+#endif
+
+	auto* Relay = new FRecorderRelay(Socket, Trace.Writeable);
+
+	uint32 IdPieces[] = {
+		Relay->GetIpAddress(),
+		Socket.remote_endpoint().port(),
+		Socket.local_endpoint().port(),
+		0,
+	};
+
+	FSession Session;
+	Session.Relay = Relay;
+	Session.Id = QuickStoreHash(IdPieces);
+	Session.TraceId = Trace.Id;
+	Sessions.Add(Session);
+
+	return true;
+}
+
+////////////////////////////////////////////////////////////////////////////////
+void FRecorder::OnTick()
+{
+	uint32 FinalNum = 0;
+	for (int i = 0, n = Sessions.Num(); i < n; ++i)
+	{
+		FSession& Session = Sessions[i];
+		if (Session.Relay->IsOpen())
+		{
+			Sessions[FinalNum] = Session;
+			++FinalNum;
+			continue;
+		}
+
+		delete Session.Relay;
+	}
+
+	Sessions.SetNum(FinalNum);
+}
+
+/* vim: set noexpandtab : */
diff --git a/Engine/Source/Programs/UnrealTraceServer/src/Recorder.h b/Engine/Source/Programs/UnrealTraceServer/src/Recorder.h
new file mode 100644
--- /dev/null
+++ b/Engine/Source/Programs/UnrealTraceServer/src/Recorder.h
@@ -0,0 +1,47 @@
+// Copyright Epic Games, Inc. All Rights Reserved.
+
+#pragma once
+
+#include "Asio.h"
+#include "AsioTcpServer.h"
+#include "AsioTickable.h"
+#include "Foundation.h"
+
+class FStore;
+class FRecorderRelay;
+
+////////////////////////////////////////////////////////////////////////////////
+class FRecorder
+	: public FAsioTcpServer
+	, public FAsioTickable
+{
+public:
+	class FSession
+	{
+	public:
+		uint32				GetId() const;
+		uint32				GetTraceId() const;
+		uint32				GetIpAddress() const;
+		uint32				GetControlPort() const;
+
+	private:
+		friend				FRecorder;
+		FRecorderRelay*		Relay;
+		uint32				Id;
+		uint32				TraceId;
+	};
+
+							FRecorder(asio::io_context& IoContext, FStore& InStore);
+							~FRecorder();
+	void					Close();
+	uint32					GetSessionCount() const;
+	const FSession*			GetSessionInfo(uint32 Index) const;
+
+private:
+	virtual bool			OnAccept(asio::ip::tcp::socket& Socket) override;
+	virtual void			OnTick() override;
+	TArray<FSession>		Sessions;
+	FStore&					Store;
+};
+
+/* vim: set noexpandtab : */
diff --git a/Engine/Source/Programs/UnrealTraceServer/src/Store.cpp b/Engine/Source/Programs/UnrealTraceServer/src/Store.cpp
new file mode 100644
--- /dev/null
+++ b/Engine/Source/Programs/UnrealTraceServer/src/Store.cpp
@@ -0,0 +1,360 @@
+// Copyright Epic Games, Inc. All Rights Reserved.
+
+#include "Pch.h"
+#include "Store.h"
+
+#if TS_USING(TS_WITH_DIR_WATCHER)
+////////////////////////////////////////////////////////////////////////////////
+#if TS_USING(TS_PLATFORM_WINDOWS)
+class FStore::FDirWatcher
+	: public asio::windows::object_handle
+{
+public:
+	using asio::windows::object_handle::object_handle;
+};
+#else
+class FStore::FDirWatcher
+	//: public asio::posix::stream_descriptor
+{
+public:
+	void async_wait(...) {}
+	void cancel() {}
+	void close() {}
+	bool is_open() { return false; }
+};
+#endif // PLATFORM_WINDOWS
+#endif // TS_WITH_DIR_WATCHER
+
+
+
+////////////////////////////////////////////////////////////////////////////////
+FStore::FTrace::FTrace(const char* InPath)
+: Path(InPath)
+{
+	// Extract the trace's name
+	const char* Dot = std::strrchr(*Path, '.');
+	if (Dot == nullptr)
+	{
+		Dot = *Path;
+	}
+
+	for (const char* c = Dot; c > *Path; --c)
+	{
+		if (c[-1] == '\\' || c[-1] == '/')
+		{
+			Name = FStringView(c, int32(Dot - c));
+			break;
+		}
+	}
+
+	Id = QuickStoreHash(Name);
+
+	// Calculate that trace's timestamp
+	static const uint64 UnixToUnrealEpochBias = 0x089f7ff5f7b58000ull;
+	std::filesystem::file_time_type LastWriteTime = std::filesystem::last_write_time(InPath);
+	Timestamp = LastWriteTime.time_since_epoch().count();
+	Timestamp += UnixToUnrealEpochBias;
+}
+
+////////////////////////////////////////////////////////////////////////////////
+const FStringView& FStore::FTrace::GetName() const
+{
+	return Name;
+}
+
+////////////////////////////////////////////////////////////////////////////////
+uint32 FStore::FTrace::GetId() const
+{
+	return Id;
+}
+
+////////////////////////////////////////////////////////////////////////////////
+uint64 FStore::FTrace::GetSize() const
+{
+	return std::filesystem::file_size(*Path);
+}
+
+////////////////////////////////////////////////////////////////////////////////
+uint64 FStore::FTrace::GetTimestamp() const
+{
+	return Timestamp;
+}
+
+
+
+////////////////////////////////////////////////////////////////////////////////
+FStore::FStore(asio::io_context& InIoContext, const char* InStoreDir)
+: IoContext(InIoContext)
+, StoreDir(InStoreDir)
+{
+	StoreDir += "/001";
+	std::filesystem::create_directories(*StoreDir);
+
+	Refresh();
+
+#if TS_USING(TS_WITH_DIR_WATCHER)
+#if TS_USING(TS_PLATFORM_WINDOWS)
+	FWinApiStr StoreDirW(*StoreDir);
+	HANDLE DirWatchHandle = FindFirstChangeNotificationW(StoreDirW, false, FILE_NOTIFY_CHANGE_FILE_NAME);
+	if (DirWatchHandle == INVALID_HANDLE_VALUE)
+	{
+		DirWatchHandle = 0;
+	}
+	DirWatcher = new FDirWatcher(IoContext, DirWatchHandle);
+#else
+	/* NOTE TO SELF - DELETE FWinApiStr NOW! */
+#endif // TS_PLATFORM_WINDOWS
+
+	WatchDir();
+#endif // TS_WITH_DIR_WATCHER
+}
+
+////////////////////////////////////////////////////////////////////////////////
+FStore::~FStore()
+{
+#if TS_USING(TS_WITH_DIR_WATCHER)
+	if (DirWatcher != nullptr)
+	{
+		check(!DirWatcher->is_open());
+		delete DirWatcher;
+	}
+#endif // TS_WITH_DIR_WATCHER
+}
+
+////////////////////////////////////////////////////////////////////////////////
+void FStore::Close()
+{
+#if TS_USING(TS_WITH_DIR_WATCHER)
+	if (DirWatcher != nullptr)
+	{
+		DirWatcher->cancel();
+		DirWatcher->close();
+	}
+#endif // TS_WITH_DIR_WATCHER
+
+	ClearTraces();
+}
+
+////////////////////////////////////////////////////////////////////////////////
+void FStore::ClearTraces()
+{
+	for (FTrace* Trace : Traces)
+	{
+		delete Trace;
+	}
+
+	Traces.Empty();
+	ChangeSerial = 0;
+}
+
+#if TS_USING(TS_WITH_DIR_WATCHER)
+////////////////////////////////////////////////////////////////////////////////
+void FStore::WatchDir()
+{
+	if (DirWatcher == nullptr)
+	{
+		return;
+	}
+
+	DirWatcher->async_wait([this] (asio::error_code ErrorCode)
+	{
+		if (ErrorCode)
+		{
+			return;
+		}
+
+#if TS_USING(TS_PLATFORM_WINDOWS)
+		FindNextChangeNotification(DirWatcher->native_handle());
+#endif
+		Refresh();
+		WatchDir();
+	});
+}
+#endif // TS_WITH_DIR_WATCHER
+
+////////////////////////////////////////////////////////////////////////////////
+const char* FStore::GetStoreDir() const
+{
+	return *StoreDir;
+}
+
+////////////////////////////////////////////////////////////////////////////////
+uint32 FStore::GetChangeSerial() const
+{
+	return ChangeSerial;
+}
+
+////////////////////////////////////////////////////////////////////////////////
+uint32 FStore::GetTraceCount() const
+{
+	return Traces.Num();
+}
+
+////////////////////////////////////////////////////////////////////////////////
+const FStore::FTrace* FStore::GetTraceInfo(uint32 Index) const
+{
+	if (Index >= uint32(Traces.Num()))
+	{
+		return nullptr;
+	}
+
+	return Traces[Index];
+}
+
+////////////////////////////////////////////////////////////////////////////////
+FStore::FTrace* FStore::GetTrace(uint32 Id) const
+{
+	for (FTrace* Trace : Traces)
+	{
+		if (Trace->GetId() == Id)
+		{
+			return Trace;
+		}
+	}
+
+	return nullptr;
+}
+
+////////////////////////////////////////////////////////////////////////////////
+FStore::FTrace* FStore::AddTrace(const char* Path)
+{
+	FTrace NewTrace(Path);
+
+	uint32 Id = NewTrace.GetId();
+	if (FTrace* Existing = GetTrace(Id))
+	{
+		return Existing;
+	}
+
+	ChangeSerial += Id;
+
+	FTrace* Trace = new FTrace(MoveTemp(NewTrace));
+	Traces.Add(Trace);
+	return Trace;
+}
+
+////////////////////////////////////////////////////////////////////////////////
+FStore::FNewTrace FStore::CreateTrace()
+{
+	FString TracePath;
+#if 0
+	bool bOk = false;
+	for (int i = 0; i < 256; ++i)
+	{
+		uint32 TraceId = ++LastTraceId;
+
+		TracePath = StoreDir;
+		TracePath.Appendf("/%05d", TraceId);
+
+		if (!std::filesystem::is_directory(*TracePath) && std::filesystem::create_directories(*TracePath))
+		{
+			bOk = true;
+			break;
+		}
+	}
+
+	if (!bOk)
+	{
+		return {};
+	}
+
+	TracePath += "/data.utrace";
+#else
+
+	// N.B. Not thread safe!?
+	char Prefix[24];
+	std::time_t Now = std::time(nullptr);
+	std::tm* LocalNow = std::localtime(&Now);
+	std::strftime(Prefix, TS_ARRAY_COUNT(Prefix), "%Y%m%d_%H%M%S", LocalNow);
+
+	TracePath = StoreDir;
+	TracePath += "/";
+	TracePath += Prefix;
+	TracePath += ".utrace";
+	for (uint32 Index = 0; std::filesystem::is_regular_file(*TracePath); ++Index)
+	{
+		char Suffix[64];
+		std::snprintf(Suffix, TS_ARRAY_COUNT(Suffix), "/%s_%02d.utrace", Prefix, Index);
+
+		TracePath = StoreDir;
+		TracePath += *Suffix;
+	}
+#endif // 0
+
+	FAsioWriteable* File = FAsioFile::WriteFile(IoContext, *TracePath);
+	if (File == nullptr)
+	{
+		return {};
+	}
+
+	FTrace* Trace = AddTrace(*TracePath);
+	if (Trace == nullptr)
+	{
+		delete File;
+		return {};
+	}
+
+	return { Trace->GetId(), File };
+}
+
+////////////////////////////////////////////////////////////////////////////////
+bool FStore::HasTrace(uint32 Id) const
+{
+	return GetTrace(Id) != nullptr;
+}
+
+////////////////////////////////////////////////////////////////////////////////
+FAsioReadable* FStore::OpenTrace(uint32 Id)
+{
+	FTrace* Trace = GetTrace(Id);
+	if (Trace == nullptr)
+	{
+		return nullptr;
+	}
+
+	FString TracePath;
+	TracePath = StoreDir;
+#if 0
+	TracePath.Appendf("/%05d/data.utrace", Id);
+#else
+	TracePath += "/";
+	TracePath += Trace->GetName();
+	TracePath += ".utrace";
+#endif // 0
+
+	return FAsioFile::ReadFile(IoContext, *TracePath);
+}
+
+////////////////////////////////////////////////////////////////////////////////
+void FStore::Refresh()
+{
+	ClearTraces();
+
+	for (auto& DirItem : std::filesystem::directory_iterator(*StoreDir))
+	{
+#if 0
+		if (!DirItem.is_directory())
+		{
+			continue;
+		}
+
+		int32 Id = FCString::Atoi(Path);
+		LastTraceId = (Id < LastTraceId) ? Id : LastTraceId;
+#else
+		if (DirItem.is_directory())
+		{
+			continue;
+		}
+
+		std::filesystem::path Extension = DirItem.path().extension();
+		if (Extension != ".utrace")
+		{
+			continue;
+		}
+
+		AddTrace(DirItem.path().string().c_str());
+#endif // 0
+	};
+}
+
+/* vim: set noexpandtab : */
diff --git a/Engine/Source/Programs/UnrealTraceServer/src/Store.h b/Engine/Source/Programs/UnrealTraceServer/src/Store.h
new file mode 100644
--- /dev/null
+++ b/Engine/Source/Programs/UnrealTraceServer/src/Store.h
@@ -0,0 +1,78 @@
+// Copyright Epic Games, Inc. All Rights Reserved.
+
+#pragma once
+
+#define TS_WITH_DIR_WATCHER TS_ON
+
+#include "Asio.h"
+#include "AsioFile.h"
+#include "Foundation.h"
+#include "Utils.h"
+
+class FAsioReadable;
+class FAsioWriteable;
+
+////////////////////////////////////////////////////////////////////////////////
+enum class EStoreVersion
+{
+	Value = 0x0100, // 0xMMmm MM=major, mm=minor
+};
+
+////////////////////////////////////////////////////////////////////////////////
+class FStore
+{
+public:
+	class FTrace
+	{
+	public:
+							FTrace(const char* InPath);
+		const FStringView&	GetName() const;
+		uint32				GetId() const;
+		uint64				GetSize() const;
+		uint64				GetTimestamp() const;
+
+	private:
+		friend				FStore;
+		FString				Path;
+		FStringView			Name;
+		uint64				Timestamp;
+		uint32				Id = 0;
+	};
+
+	struct FNewTrace
+	{
+		uint32			Id;
+		FAsioWriteable* Writeable;
+	};
+
+						FStore(asio::io_context& IoContext, const char* InStoreDir);
+						~FStore();
+	void				Close();
+	const char*			GetStoreDir() const;
+	uint32				GetChangeSerial() const;
+	uint32				GetTraceCount() const;
+	const FTrace*		GetTraceInfo(uint32 Index) const;
+	bool				HasTrace(uint32 Id) const;
+	FNewTrace			CreateTrace();
+	FAsioReadable*		OpenTrace(uint32 Id);
+
+private:
+	FTrace*				GetTrace(uint32 Id) const;
+	FTrace*				AddTrace(const char* Path);
+	void				ClearTraces();
+	void				Refresh();
+	asio::io_context&	IoContext;
+	FString				StoreDir;
+	TArray<FTrace*>		Traces;
+	uint32				ChangeSerial;
+#if TS_USING(TS_WITH_DIR_WATCHER)
+	class				FDirWatcher;
+	void				WatchDir();
+	FDirWatcher*		DirWatcher = nullptr;
+#endif
+#if 0
+	int32				LastTraceId = -1;
+#endif // 0
+};
+
+/* vim: set noexpandtab : */
diff --git a/Engine/Source/Programs/UnrealTraceServer/src/StoreCborServer.cpp b/Engine/Source/Programs/UnrealTraceServer/src/StoreCborServer.cpp
new file mode 100644
--- /dev/null
+++ b/Engine/Source/Programs/UnrealTraceServer/src/StoreCborServer.cpp
@@ -0,0 +1,441 @@
+// Copyright Epic Games, Inc. All Rights Reserved.
+
+#include "Pch.h"
+#include "AsioIoable.h"
+#include "AsioSocket.h"
+#include "CborPayload.h"
+#include "Recorder.h"
+#include "Store.h"
+#include "StoreCborServer.h"
+#include "TraceRelay.h"
+
+////////////////////////////////////////////////////////////////////////////////
+class FStoreCborPeer
+	: public FAsioIoSink
+{
+public:
+							FStoreCborPeer(asio::ip::tcp::socket& InSocket, FStore& InStore, FRecorder& InRecorder, FStoreCborServer& InParent);
+	virtual					~FStoreCborPeer();
+	bool					IsOpen() const;
+	void					Close();
+
+protected:
+	void					OnPayload();
+	void					OnSessionCount();
+	void					OnSessionInfo();
+	void					OnStatus();
+	void					OnTraceCount();
+	void					OnTraceInfo();
+	void					OnTraceRead();
+	void					SendError(EStatusCode StatusCode);
+	void					SendResponse(const FPayload& Payload);
+	virtual void			OnIoComplete(uint32 Id, int32 Size) override;
+	enum					{ OpStart, OpReadPayloadSize, OpReadPayload, OpSendResponse };
+	FStore&					Store;
+	FRecorder&				Recorder;
+	FStoreCborServer&		Parent;
+	FAsioSocket				Socket;
+	uint32					PayloadSize;
+	FResponse				Response;
+};
+
+////////////////////////////////////////////////////////////////////////////////
+FStoreCborPeer::FStoreCborPeer(
+	asio::ip::tcp::socket& InSocket,
+	FStore& InStore,
+	FRecorder& InRecorder,
+	FStoreCborServer& InParent)
+: Store(InStore)
+, Recorder(InRecorder)
+, Parent(InParent)
+, Socket(InSocket)
+{
+	OnIoComplete(OpStart, 0);
+}
+
+////////////////////////////////////////////////////////////////////////////////
+FStoreCborPeer::~FStoreCborPeer()
+{
+	check(!IsOpen());
+}
+
+////////////////////////////////////////////////////////////////////////////////
+bool FStoreCborPeer::IsOpen() const
+{
+	return Socket.IsOpen();
+}
+
+////////////////////////////////////////////////////////////////////////////////
+void FStoreCborPeer::Close()
+{
+	Socket.Close();
+}
+
+////////////////////////////////////////////////////////////////////////////////
+void FStoreCborPeer::OnSessionCount()
+{
+	TPayloadBuilder<> Builder(EStatusCode::Success);
+	Builder.AddInteger("count", Recorder.GetSessionCount());
+	SendResponse(Builder.Done());
+}
+
+////////////////////////////////////////////////////////////////////////////////
+void FStoreCborPeer::OnSessionInfo()
+{
+	const FRecorder::FSession* Session = nullptr;
+
+	int32 Index = int32(Response.GetInteger("index", -1));
+	if (Index >= 0)
+	{
+		Session = Recorder.GetSessionInfo(Index);
+	}
+	else if (uint32 Id = uint32(Response.GetInteger("id", 0)))
+	{
+		for (int i = 0, n = Recorder.GetSessionCount(); i < n; ++i)
+		{
+			const FRecorder::FSession* Candidate = Recorder.GetSessionInfo(i);
+			if (Candidate->GetId() == Id)
+			{
+				Session = Candidate;
+				break;
+			}
+		}
+	}
+	else if (uint32 TraceId = uint32(Response.GetInteger("trace_id", 0)))
+	{
+		for (int i = 0, n = Recorder.GetSessionCount(); i < n; ++i)
+		{
+			const FRecorder::FSession* Candidate = Recorder.GetSessionInfo(i);
+			if (Candidate->GetTraceId() == TraceId)
+			{
+				Session = Candidate;
+				break;
+			}
+		}
+	}
+
+	if (Session == nullptr)
+	{
+		return SendError(EStatusCode::BadRequest);
+	}
+
+	TPayloadBuilder<> Builder(EStatusCode::Success);
+	Builder.AddInteger("id", Session->GetId());
+	Builder.AddInteger("trace_id", Session->GetTraceId());
+	Builder.AddInteger("ip_address", Session->GetIpAddress());
+	Builder.AddInteger("control_port", Session->GetControlPort());
+	SendResponse(Builder.Done());
+}
+
+////////////////////////////////////////////////////////////////////////////////
+void FStoreCborPeer::OnStatus()
+{
+	TPayloadBuilder<> Builder(EStatusCode::Success);
+	Builder.AddInteger("recorder_port", Recorder.GetPort());
+	Builder.AddInteger("change_serial", Store.GetChangeSerial());
+	Builder.AddString("store_dir", Store.GetStoreDir());
+	SendResponse(Builder.Done());
+}
+
+////////////////////////////////////////////////////////////////////////////////
+void FStoreCborPeer::OnTraceCount()
+{
+	TPayloadBuilder<> Builder(EStatusCode::Success);
+	Builder.AddInteger("count", Store.GetTraceCount());
+	SendResponse(Builder.Done());
+}
+
+////////////////////////////////////////////////////////////////////////////////
+void FStoreCborPeer::OnTraceInfo()
+{
+	const FStore::FTrace* Trace = nullptr;
+
+	int32 Index = int32(Response.GetInteger("index", -1));
+	if (Index >= 0)
+	{
+		Trace = Store.GetTraceInfo(Index);
+	}
+	else
+	{
+		uint32 Id = uint32(Response.GetInteger("id", 0));
+		if (Id != 0)
+		{
+			for (int i = 0, n = Store.GetTraceCount(); i < n; ++i)
+			{
+				const FStore::FTrace* Candidate = Store.GetTraceInfo(i);
+				if (Candidate->GetId() == Id)
+				{
+					Trace = Candidate;
+					break;
+				}
+			}
+		}
+	}
+
+	if (Trace == nullptr)
+	{
+		return SendError(EStatusCode::BadRequest);
+	}
+
+	char OutName[256];
+	const FStringView& Name = Trace->GetName();
+	uint32 NameLength = std::min(uint32(sizeof(OutName)), uint32(Name.Len()));
+	for (uint32 i = 0; i < NameLength; ++i)
+	{
+		OutName[i] = char(Name[i]);
+	}
+
+	TPayloadBuilder<> Builder(EStatusCode::Success);
+	Builder.AddInteger("id", Trace->GetId());
+	Builder.AddInteger("size", Trace->GetSize());
+	Builder.AddInteger("timestamp", Trace->GetTimestamp());
+	Builder.AddString("name", OutName, NameLength);
+	SendResponse(Builder.Done());
+}
+
+////////////////////////////////////////////////////////////////////////////////
+void FStoreCborPeer::OnTraceRead()
+{
+	uint32 Id = uint32(Response.GetInteger("id", 0));
+	if (!Id || !Store.HasTrace(Id))
+	{
+		return SendError(EStatusCode::BadRequest);
+	}
+
+	FTraceRelay* Relay = Parent.RelayTrace(Id);
+	if (Relay == nullptr)
+	{
+		return SendError(EStatusCode::InternalError);
+	}
+
+	TPayloadBuilder<> Builder(EStatusCode::Success);
+	Builder.AddInteger("port", Relay->GetPort());
+	SendResponse(Builder.Done());
+}
+
+////////////////////////////////////////////////////////////////////////////////
+void FStoreCborPeer::OnPayload()
+{
+	FStringView Request = Response.GetString("$request", "");
+	FStringView Path = Response.GetString("$path", "");
+	if (!Request.Len() || !Path.Len())
+	{
+		SendError(EStatusCode::BadRequest);
+		return;
+	}
+
+	static struct {
+		uint32	Hash;
+		void	(FStoreCborPeer::*Func)();
+	} const DispatchTable[] = {
+		{ QuickStoreHash("v1/session/count"),	&FStoreCborPeer::OnSessionCount },
+		{ QuickStoreHash("v1/session/info"),	&FStoreCborPeer::OnSessionInfo },
+		{ QuickStoreHash("v1/status"),			&FStoreCborPeer::OnStatus },
+		{ QuickStoreHash("v1/trace/count"),		&FStoreCborPeer::OnTraceCount },
+		{ QuickStoreHash("v1/trace/info"),		&FStoreCborPeer::OnTraceInfo },
+		{ QuickStoreHash("v1/trace/read"),		&FStoreCborPeer::OnTraceRead },
+	};
+
+	uint32 PathHash = QuickStoreHash(Path);
+	for (const auto& Row : DispatchTable)
+	{
+		if (Row.Hash == PathHash)
+		{
+			return (this->*(Row.Func))();
+		}
+	}
+
+	SendError(EStatusCode::NotFound);
+}
+
+////////////////////////////////////////////////////////////////////////////////
+void FStoreCborPeer::SendError(EStatusCode StatusCode)
+{
+	TPayloadBuilder<> Builder(StatusCode);
+	FPayload Payload = Builder.Done();
+	SendResponse(Payload);
+}
+
+////////////////////////////////////////////////////////////////////////////////
+void FStoreCborPeer::SendResponse(const FPayload& Payload)
+{
+	struct {
+		uint32	Size;
+		uint8	Data[];
+	}* Dest;
+
+	Dest = decltype(Dest)(Response.Reserve(sizeof(*Dest) + Payload.Size));
+	Dest->Size = Payload.Size;
+	memcpy(Dest->Data, Payload.Data, Payload.Size);
+	Socket.Write(Dest, sizeof(*Dest) + Payload.Size, this, OpSendResponse);
+}
+
+////////////////////////////////////////////////////////////////////////////////
+void FStoreCborPeer::OnIoComplete(uint32 Id, int32 Size)
+{
+	if (Size < 0)
+	{
+		Close();
+		return;
+	}
+
+	switch (Id)
+	{
+	case OpStart:
+	case OpSendResponse:
+		Socket.Read(&PayloadSize, sizeof(uint32), this, OpReadPayloadSize);
+		break;
+
+	case OpReadPayloadSize:
+		Socket.Read(Response.Reserve(PayloadSize), PayloadSize, this, OpReadPayload);
+		break;
+
+	case OpReadPayload:
+		OnPayload();
+		break;
+	}
+}
+
+
+
+////////////////////////////////////////////////////////////////////////////////
+FStoreCborServer::FStoreCborServer(
+	asio::io_context& IoContext,
+	int32 Port,
+	FStore& InStore,
+	FRecorder& InRecorder)
+: FAsioTcpServer(IoContext)
+, FAsioTickable(IoContext)
+, Store(InStore)
+, Recorder(InRecorder)
+{
+	if (!StartServer(Port))
+	{
+		StartServer();
+	}
+	StartTick(500);
+}
+
+////////////////////////////////////////////////////////////////////////////////
+FStoreCborServer::~FStoreCborServer()
+{
+	check(!FAsioTcpServer::IsOpen());
+	check(!FAsioTickable::IsActive());
+
+	for (FStoreCborPeer* Peer : Peers)
+	{
+		delete Peer;
+	}
+
+	for (FTraceRelay* Relay : Relays)
+	{
+		delete Relay;
+	}
+}
+
+////////////////////////////////////////////////////////////////////////////////
+void FStoreCborServer::Close()
+{
+	FAsioTcpServer::Close();
+	FAsioTickable::StopTick();
+
+	for (FStoreCborPeer* Peer : Peers)
+	{
+		Peer->Close();
+	}
+
+	for (FTraceRelay* Relay : Relays)
+	{
+		Relay->Close();
+	}
+}
+
+////////////////////////////////////////////////////////////////////////////////
+FStore& FStoreCborServer::GetStore() const
+{
+	return Store;
+}
+
+////////////////////////////////////////////////////////////////////////////////
+FRecorder& FStoreCborServer::GetRecorder() const
+{
+	return Recorder;
+}
+
+////////////////////////////////////////////////////////////////////////////////
+bool FStoreCborServer::OnAccept(asio::ip::tcp::socket& Socket)
+{
+	FStoreCborPeer* Peer = new FStoreCborPeer(Socket, Store, Recorder, *this);
+	Peers.Add(Peer);
+	return true;
+}
+
+////////////////////////////////////////////////////////////////////////////////
+void FStoreCborServer::OnTick()
+{
+	// Clean up dead peers
+	uint32 FinalNum = 0;
+	for (int i = 0, n = Peers.Num(); i < n; ++i)
+	{
+		FStoreCborPeer* Peer = Peers[i];
+		if (Peer->IsOpen())
+		{
+			Peers[FinalNum] = Peer;
+			++FinalNum;
+			continue;
+		}
+
+		delete Peer;
+	}
+	Peers.SetNum(FinalNum);
+
+	// Clean up dead relays
+	FinalNum = 0;
+	for (int i = 0, n = Relays.Num(); i < n; ++i)
+	{
+		FTraceRelay* Relay = Relays[i];
+		if (Relay->IsOpen())
+		{
+			Relays[FinalNum] = Relay;
+			++FinalNum;
+			continue;
+		}
+
+		delete Relay;
+	}
+	Relays.SetNum(FinalNum);
+}
+
+////////////////////////////////////////////////////////////////////////////////
+FTraceRelay* FStoreCborServer::RelayTrace(uint32 Id)
+{
+	FAsioReadable* Input = Store.OpenTrace(Id);
+	if (Input == nullptr)
+	{
+		return nullptr;
+	}
+
+	uint32 SessionId = 0;
+	for (int i = 0, n = Recorder.GetSessionCount(); i < n; ++i)
+	{
+		const FRecorder::FSession* Session = Recorder.GetSessionInfo(i);
+		if (Session->GetTraceId() == Id)
+		{
+			SessionId = Session->GetId();
+			break;
+		}
+	}
+
+	asio::io_context& IoContext = FAsioTcpServer::GetIoContext();
+	FTraceRelay* Relay = new FTraceRelay(IoContext, Input, SessionId, Recorder);
+	if (Relay->GetPort() == 0)
+	{
+		delete Relay;
+		return nullptr;
+	}
+
+	Relays.Add(Relay);
+	return Relay;
+}
+
+/* vim: set noexpandtab : */
diff --git a/Engine/Source/Programs/UnrealTraceServer/src/StoreCborServer.h b/Engine/Source/Programs/UnrealTraceServer/src/StoreCborServer.h
new file mode 100644
--- /dev/null
+++ b/Engine/Source/Programs/UnrealTraceServer/src/StoreCborServer.h
@@ -0,0 +1,37 @@
+// Copyright Epic Games, Inc. All Rights Reserved.
+
+#pragma once
+
+#include "Asio.h"
+#include "AsioTcpServer.h"
+#include "AsioTickable.h"
+#include "Foundation.h"
+
+class FRecorder;
+class FStore;
+class FStoreCborPeer;
+class FTraceRelay;
+
+////////////////////////////////////////////////////////////////////////////////
+class FStoreCborServer
+	: public FAsioTcpServer
+	, public FAsioTickable
+{
+public:
+							FStoreCborServer(asio::io_context& IoContext, int32 Port, FStore& InStore, FRecorder& InRecorder);
+							~FStoreCborServer();
+	void					Close();
+	FStore&					GetStore() const;
+	FRecorder&				GetRecorder() const;
+	FTraceRelay*			RelayTrace(uint32 Id);
+
+private:
+	virtual bool			OnAccept(asio::ip::tcp::socket& Socket) override;
+	virtual void			OnTick() override;
+	TArray<FStoreCborPeer*>	Peers;
+	TArray<FTraceRelay*>	Relays;
+	FStore&					Store;
+	FRecorder&				Recorder;
+};
+
+/* vim: set noexpandtab : */
diff --git a/Engine/Source/Programs/UnrealTraceServer/src/StoreService.cpp b/Engine/Source/Programs/UnrealTraceServer/src/StoreService.cpp
new file mode 100644
--- /dev/null
+++ b/Engine/Source/Programs/UnrealTraceServer/src/StoreService.cpp
@@ -0,0 +1,91 @@
+// Copyright Epic Games, Inc. All Rights Reserved.
+
+#include "Pch.h"
+#include "Asio.h"
+#include "AsioContext.h"
+#include "Recorder.h"
+#include "Store.h"
+#include "StoreCborServer.h"
+#include "StoreService.h"
+
+////////////////////////////////////////////////////////////////////////////////
+struct FStoreServiceImpl
+{
+public:
+							FStoreServiceImpl(const FStoreService::FDesc& Desc);
+							~FStoreServiceImpl();
+	FAsioContext			Context;
+	FStore					Store;
+	FRecorder				Recorder;
+	FStoreCborServer		CborServer;
+};
+
+////////////////////////////////////////////////////////////////////////////////
+FStoreServiceImpl::FStoreServiceImpl(const FStoreService::FDesc& Desc)
+: Context(Desc.ThreadCount)
+, Store(Context.Get(), Desc.StoreDir)
+, Recorder(Context.Get(), Store)
+, CborServer(Context.Get(), Desc.StorePort, Store, Recorder)
+{
+}
+
+////////////////////////////////////////////////////////////////////////////////
+FStoreServiceImpl::~FStoreServiceImpl()
+{
+	asio::post(Context.Get(), [this] () {
+		CborServer.Close();
+		Recorder.Close();
+		Store.Close();
+	});
+	Context.Wait();
+}
+
+
+
+////////////////////////////////////////////////////////////////////////////////
+FStoreService* FStoreService::Create(const FDesc& InDesc)
+{
+	FDesc Desc = InDesc;
+
+	if (Desc.ThreadCount <= 0)
+	{
+		Desc.ThreadCount = std::thread::hardware_concurrency();
+	}
+
+	// TODO: not thread safe yet
+	Desc.ThreadCount = 1;
+
+	FStoreServiceImpl* Impl = new FStoreServiceImpl(Desc);
+	if (Desc.RecorderPort >= 0)
+	{
+		FRecorder& Recorder = Impl->Recorder;
+		Recorder.StartServer(Desc.RecorderPort);
+	}
+
+	Impl->Context.Start();
+
+	return (FStoreService*)Impl;
+}
+
+////////////////////////////////////////////////////////////////////////////////
+void FStoreService::operator delete (void* Addr)
+{
+	auto* Self = (FStoreServiceImpl*)Addr;
+	delete Self;
+}
+
+////////////////////////////////////////////////////////////////////////////////
+uint32 FStoreService::GetPort() const
+{
+	auto* Self = (FStoreServiceImpl*)this;
+	return Self->CborServer.GetPort();
+}
+
+////////////////////////////////////////////////////////////////////////////////
+uint32 FStoreService::GetRecorderPort() const
+{
+	auto* Self = (FStoreServiceImpl*)this;
+	return Self->Recorder.GetPort();
+}
+
+/* vim: set noexpandtab : */
diff --git a/Engine/Source/Programs/UnrealTraceServer/src/StoreService.h b/Engine/Source/Programs/UnrealTraceServer/src/StoreService.h
new file mode 100644
--- /dev/null
+++ b/Engine/Source/Programs/UnrealTraceServer/src/StoreService.h
@@ -0,0 +1,31 @@
+// Copyright Epic Games, Inc. All Rights Reserved.
+
+#pragma once
+
+////////////////////////////////////////////////////////////////////////////////
+class FStoreService
+{
+public:
+	struct FDesc
+	{
+		const char*			StoreDir;
+		int32				StorePort		= 0; // <=0:auto-assign
+		int32				RecorderPort	= 0; // 0:auto-assign, -1:off
+		int32				ThreadCount		= 0; // <=0:logical CPU count
+	};
+
+							~FStoreService() = default;
+	static FStoreService*	Create(const FDesc& Desc);
+	void					operator delete (void* Addr);
+	uint32					GetPort() const;
+	uint32					GetRecorderPort() const;
+
+private:
+							FStoreService() = default;
+							FStoreService(const FStoreService&) = delete;
+							FStoreService(const FStoreService&&) = delete;
+	void					operator = (const FStoreService&) = delete;
+	void					operator = (const FStoreService&&) = delete;
+};
+
+/* vim: set noexpandtab : */
diff --git a/Engine/Source/Programs/UnrealTraceServer/src/TraceRelay.cpp b/Engine/Source/Programs/UnrealTraceServer/src/TraceRelay.cpp
new file mode 100644
--- /dev/null
+++ b/Engine/Source/Programs/UnrealTraceServer/src/TraceRelay.cpp
@@ -0,0 +1,124 @@
+// Copyright Epic Games, Inc. All Rights Reserved.
+
+#include "Pch.h"
+#include "AsioSocket.h"
+#include "Recorder.h"
+#include "TraceRelay.h"
+
+////////////////////////////////////////////////////////////////////////////////
+FTraceRelay::FTraceRelay(
+	asio::io_context& IoContext,
+	FAsioReadable* InInput,
+	uint32 InSessionId,
+	FRecorder& InRecorder)
+: FAsioTcpServer(IoContext)
+, FAsioTickable(IoContext)
+, Input(InInput)
+, Recorder(InRecorder)
+, SessionId(InSessionId)
+{
+	StartServer();
+	TickOnce(3000);
+}
+
+////////////////////////////////////////////////////////////////////////////////
+FTraceRelay::~FTraceRelay()
+{
+	check(!Input->IsOpen());
+	delete Input;
+
+	if (Output != nullptr)
+	{
+		check(!Output->IsOpen());
+		delete Output;
+	}
+}
+
+////////////////////////////////////////////////////////////////////////////////
+void FTraceRelay::Close()
+{
+	if (Output != nullptr)
+	{
+		Output->Close();
+	}
+
+	Input->Close();
+
+	FAsioTcpServer::Close();
+}
+
+////////////////////////////////////////////////////////////////////////////////
+bool FTraceRelay::IsOpen() const
+{
+	if (FAsioTcpServer::IsOpen())
+	{
+		return true;
+	}
+
+	bool bOpen = Input->IsOpen();
+	if (Output != nullptr)
+	{
+		bOpen &= Output->IsOpen();
+	}
+
+	return bOpen;
+}
+
+////////////////////////////////////////////////////////////////////////////////
+bool FTraceRelay::OnAccept(asio::ip::tcp::socket& Socket)
+{
+	Output = new FAsioSocket(Socket);
+	StopTick();
+	OnIoComplete(OpStart, 0);
+	return false;
+}
+
+////////////////////////////////////////////////////////////////////////////////
+void FTraceRelay::OnTick()
+{
+	if (FAsioTcpServer::IsOpen())
+	{
+		Close();
+		return;
+	}
+
+	return OnIoComplete(OpStart, 0);
+}
+
+////////////////////////////////////////////////////////////////////////////////
+void FTraceRelay::OnIoComplete(uint32 Id, int32 Size)
+{
+	if (Size < 0)
+	{
+		if (Id == OpRead && SessionId && -Size == asio::error::eof)
+		{
+			for (int i = 0, n = Recorder.GetSessionCount(); i < n; ++i)
+			{
+				const FRecorder::FSession* Session = Recorder.GetSessionInfo(i);
+				if (Session->GetId() == SessionId)
+				{
+					TickOnce(200);
+					return;
+				}
+			}
+		}
+
+		Output->Close();
+		Input->Close();
+		return;
+	}
+
+	switch (Id)
+	{
+	case OpStart:
+	case OpSend:
+		Input->Read(Buffer, BufferSize, this, OpRead);
+		break;
+
+	case OpRead:
+		Output->Write(Buffer, Size, this, OpSend);
+		break;
+	}
+}
+
+/* vim: set noexpandtab : */
diff --git a/Engine/Source/Programs/UnrealTraceServer/src/TraceRelay.h b/Engine/Source/Programs/UnrealTraceServer/src/TraceRelay.h
new file mode 100644
--- /dev/null
+++ b/Engine/Source/Programs/UnrealTraceServer/src/TraceRelay.h
@@ -0,0 +1,38 @@
+// Copyright Epic Games, Inc. All Rights Reserved.
+
+#pragma once
+
+#include "Asio.h"
+#include "AsioIoable.h"
+#include "AsioTcpServer.h"
+#include "AsioTickable.h"
+
+class FRecorder;
+class FAsioSocket;
+
+////////////////////////////////////////////////////////////////////////////////
+class FTraceRelay
+	: public FAsioTcpServer
+	, public FAsioIoSink
+	, public FAsioTickable
+{
+public:
+						FTraceRelay(asio::io_context& IoContext, FAsioReadable* InInput, uint32 InTraceIf, FRecorder& InRecorder);
+						~FTraceRelay();
+	bool				IsOpen() const;
+	void				Close();
+
+private:
+	virtual bool		OnAccept(asio::ip::tcp::socket& Socket) override;
+	virtual void		OnTick() override;
+	virtual void		OnIoComplete(uint32 Id, int32 Size) override;
+	enum				{ OpStart, OpRead, OpSend };
+	static const uint32	BufferSize = 64 << 10;
+	FAsioReadable*		Input;
+	FAsioSocket*		Output = nullptr;
+	FRecorder&			Recorder;
+	uint32				SessionId;
+	uint8				Buffer[BufferSize];
+};
+
+/* vim: set noexpandtab : */
diff --git a/Engine/Source/Programs/UnrealTraceServer/src/Utils.h b/Engine/Source/Programs/UnrealTraceServer/src/Utils.h
new file mode 100644
--- /dev/null
+++ b/Engine/Source/Programs/UnrealTraceServer/src/Utils.h
@@ -0,0 +1,198 @@
+// Copyright Epic Games, Inc. All Rights Reserved.
+
+#pragma once
+
+#include "Foundation.h"
+
+////////////////////////////////////////////////////////////////////////////////
+class FInlineBufferBase
+{
+public:
+				FInlineBufferBase() = default;
+				FInlineBufferBase(uint32 InlineMax);
+	uint32		GetSize() const;
+	bool		HasOverflowed() const;
+
+protected:
+	void		SetMax(uint32 NewMax);
+	uint32		GetMax() const;
+	uint32		Used = 0;
+	int32		Max;
+};
+
+////////////////////////////////////////////////////////////////////////////////
+inline FInlineBufferBase::FInlineBufferBase(uint32 InlineMax)
+: Max(int32(~InlineMax))
+{
+}
+
+////////////////////////////////////////////////////////////////////////////////
+inline uint32 FInlineBufferBase::GetSize() const
+{
+	return Used;
+}
+
+////////////////////////////////////////////////////////////////////////////////
+inline bool FInlineBufferBase::HasOverflowed() const
+{
+	return Max > 0;
+}
+
+////////////////////////////////////////////////////////////////////////////////
+inline void FInlineBufferBase::SetMax(uint32 NewMax)
+{
+	Max = NewMax;
+}
+
+////////////////////////////////////////////////////////////////////////////////
+inline uint32 FInlineBufferBase::GetMax() const
+{
+	return HasOverflowed() ? Max : ~Max;
+}
+
+
+
+////////////////////////////////////////////////////////////////////////////////
+class FInlineBuffer
+	: public FInlineBufferBase
+{
+public:
+	uint8*			Append(uint32 AppendSize);
+	void			Append(const void* Data, uint32 DataSize);
+	const uint8*	GetData() const;
+
+private:
+	uint8			Buffer[];
+};
+
+////////////////////////////////////////////////////////////////////////////////
+inline uint8* FInlineBuffer::Append(uint32 AppendSize)
+{
+	auto GrowthPolicy = [] (uint32 Size, uint32 Demand) -> uint32 {
+		uint32 NewSize = Size + 256;
+		if (Demand > NewSize)
+		{
+			NewSize = Demand + 256;
+		}
+		return NewSize;
+	};
+
+	auto** Ptr = (uint8**)Buffer;
+	uint32 CurrentMax = GetMax();
+	uint32 NextUsed = Used + AppendSize;
+	if (CurrentMax < NextUsed)
+	{
+		uint32 NewMax = GrowthPolicy(CurrentMax, NextUsed);
+		if (HasOverflowed())
+		{
+			Ptr[0] = (uint8*)realloc(Ptr[0], NewMax);
+		}
+		else
+		{
+			uint8* NewBuffer = (uint8*)malloc(NewMax);
+			memcpy(NewBuffer, Buffer, Used);
+			Ptr[0] = NewBuffer;
+		}
+		SetMax(NewMax);
+	}
+
+	uint8* Cursor = HasOverflowed() ? Ptr[0] : Buffer;
+	Cursor += Used;
+	Used = NextUsed;
+	return Cursor;
+}
+
+////////////////////////////////////////////////////////////////////////////////
+inline void FInlineBuffer::Append(const void* Data, uint32 DataSize)
+{
+	uint8* Cursor = Append(DataSize);
+	memcpy(Cursor, Data, DataSize);
+}
+
+////////////////////////////////////////////////////////////////////////////////
+inline const uint8* FInlineBuffer::GetData() const
+{
+	return HasOverflowed() ? *(uint8**)Buffer : Buffer;
+}
+
+
+
+////////////////////////////////////////////////////////////////////////////////
+template <int Size>
+class TInlineBuffer
+	: protected FInlineBufferBase
+{
+public:
+					TInlineBuffer();
+					~TInlineBuffer();
+					operator FInlineBuffer& ();
+	FInlineBuffer*	operator -> ();
+
+private:
+	enum { BufferSize = (Size >= 16) ? Size : 16 };
+
+	union
+	{
+		uint8*		Ptr;
+		uint8		Buffer[BufferSize];
+	};
+};
+
+////////////////////////////////////////////////////////////////////////////////
+template <int Size>
+TInlineBuffer<Size>::TInlineBuffer()
+: FInlineBufferBase(BufferSize)
+{
+}
+
+////////////////////////////////////////////////////////////////////////////////
+template <int Size>
+TInlineBuffer<Size>::~TInlineBuffer()
+{
+	if (HasOverflowed())
+	{
+		delete Ptr;
+	}
+}
+
+////////////////////////////////////////////////////////////////////////////////
+template <int Size>
+TInlineBuffer<Size>::operator FInlineBuffer& ()
+{
+	return *(FInlineBuffer*)this;
+}
+
+////////////////////////////////////////////////////////////////////////////////
+template <int Size>
+FInlineBuffer* TInlineBuffer<Size>::operator -> ()
+{
+	return (FInlineBuffer*)this;
+}
+
+
+
+////////////////////////////////////////////////////////////////////////////////
+template <typename Type>
+inline uint32 QuickStoreHash(const Type* String)
+{
+	uint32 Value = 5381;
+	for (; *String; ++String)
+	{
+		Value = ((Value << 5) + Value) + *String;
+	}
+	return Value;
+}
+
+////////////////////////////////////////////////////////////////////////////////
+inline uint32 QuickStoreHash(FStringView View)
+{
+	const char* String = View.GetData();
+	uint32 Value = 5381;
+	for (int i = View.Len(); i > 0; --i, ++String)
+	{
+		Value = ((Value << 5) + Value) + *String;
+	}
+	return Value;
+}
+
+/* vim: set noexpandtab : */
diff --git a/Engine/Source/Programs/UnrealTraceServer/src/Version.h b/Engine/Source/Programs/UnrealTraceServer/src/Version.h
new file mode 100644
--- /dev/null
+++ b/Engine/Source/Programs/UnrealTraceServer/src/Version.h
@@ -0,0 +1,13 @@
+// Copyright Epic Games, Inc. All Rights Reserved.
+
+#pragma once
+
+#if !defined(TS_VERSION_PROTOCOL)
+#	define TS_VERSION_PROTOCOL	1
+#endif
+
+#if !defined(TS_VERSION_MINOR)
+#	define TS_VERSION_MINOR		2
+#endif
+
+/* vim: set noexpandtab : */
diff --git a/Engine/Source/Programs/UnrealTraceServer/src/Version.rc b/Engine/Source/Programs/UnrealTraceServer/src/Version.rc
new file mode 100644
--- /dev/null
+++ b/Engine/Source/Programs/UnrealTraceServer/src/Version.rc
@@ -0,0 +1,24 @@
+// Copyright Epic Games, Inc. All Rights Reserved.
+
+#include "Version.h"
+
+1 VERSIONINFO
+	FILEVERSION		TS_VERSION_PROTOCOL, TS_VERSION_MINOR
+	PRODUCTVERSION	5 /* UE5 */
+{
+	BLOCK "StringFileInfo" {
+		BLOCK "040904b0" {
+			VALUE "CompanyName",	"Epic Games"
+			VALUE "ProductName",	"Unreal Engine Trace Server"
+			VALUE "FileDescription","Hub for UE5 runtime traces (i.e. -tracehost=<host>) and tools like Unreal Insights"
+			VALUE "ProductVersion",	"UE5"
+		}
+	}
+	BLOCK "VarFileInfo" {
+		VALUE "Translation", 0x0409, 1200
+	}
+}
+
+2 ICON "Unreal.ico"
+
+/* vim: set noexpandtab : */
