More descriptive return code to more simply identify launch issues
267772d / Martin Ridgers / [REDACTED_PHONE]:56:51 +0200

diff --git a/Engine/Source/Programs/UnrealTraceServer/src/Main.cpp b/Engine/Source/Programs/UnrealTraceServer/src/Main.cpp
--- a/Engine/Source/Programs/UnrealTraceServer/src/Main.cpp
+++ b/Engine/Source/Programs/UnrealTraceServer/src/Main.cpp
@@ -189,6 +189,12 @@ static const wchar_t*	GQuitEventName				= L"Local\\UnrealTraceEvent";
 static const wchar_t*	GBegunEventName				= L"Local\\UnrealTraceEventBegun";
 static int				MainDaemon(int, char**);
 
+////////////////////////////////////////////////////////////////////////////////
+static int CreateExitCode(uint32 Id)
+{
+	return (GetLastError() & 0xfff) | (Id << 12);
+}
+
 ////////////////////////////////////////////////////////////////////////////////
 static int MainFork(int ArgC, char** ArgV)
 {
@@ -219,7 +225,7 @@ static int MainFork(int ArgC, char** ArgV)
 		FWinHandle ProcHandle = OpenProcess(OpenProcessFlags, FALSE, InstanceInfo->Pid);
 		if (!ProcHandle.IsValid())
 		{
-			return 5;
+			return CreateExitCode(1);
 		}
 
 		SetEvent(QuitEvent);
@@ -236,7 +242,7 @@ static int MainFork(int ArgC, char** ArgV)
 	if (GetLastError() == ERROR_INSUFFICIENT_BUFFER)
 	{
 		// This should never really happen...
-		return 1;
+		return CreateExitCode(3);
 	}
 
 	// Calculate where to store the binaries.
@@ -267,7 +273,7 @@ static int MainFork(int ArgC, char** ArgV)
 		TempPath += Buffer;
 		if (!std::filesystem::copy_file(BinPath, TempPath))
 		{
-			return 3;
+			return CreateExitCode(5);
 		}
 
 		// Move the file into place. If this fails because the file exists then
@@ -277,7 +283,7 @@ static int MainFork(int ArgC, char** ArgV)
 		if (ErrorCode)
 		{
 			bool RaceLost = (ErrorCode == std::errc::file_exists);
-			return RaceLost ? 0 : 4;
+			return RaceLost ? 0 : CreateExitCode(7);
 		}
 	}
 
@@ -285,7 +291,7 @@ static int MainFork(int ArgC, char** ArgV)
 	FWinHandle BegunEvent = CreateEventW(nullptr, TRUE, FALSE, GBegunEventName);
 	if (GetLastError() == ERROR_ALREADY_EXISTS)
 	{
-		return 7;
+		return CreateExitCode(11);
 	}
 
 	// For debugging ease and consistency we will daemonize in this process
@@ -301,14 +307,14 @@ static int MainFork(int ArgC, char** ArgV)
 
 	if (bOk == FALSE)
 	{
-		return 2;
+		return CreateExitCode(13);
 	}
 #endif // TS_BUILD_DEBUG
 
 	int Ret = 0;
 	if (WaitForSingleObject(BegunEvent, 5000) == WAIT_TIMEOUT)
 	{
-		Ret = 8;
+		Ret = CreateExitCode(17);
 	}
 
 #if TS_USING(TS_BUILD_DEBUG)
@@ -329,7 +335,7 @@ static int MainDaemon(int ArgC, char** ArgV)
 		PAGE_READWRITE, 0, GIpcSize, GIpcName);
 	if (!IpcHandle.IsValid())
 	{
-		return 1;
+		return CreateExitCode(129);
 	}
 
 	// Create a named event so others can tell us to quit.
@@ -338,7 +344,7 @@ static int MainDaemon(int ArgC, char** ArgV)
 	{
 		// This really should not happen. It is expected that only one process
 		// will get this far (gated by the shared-memory object creation).
-		return 2;
+		return CreateExitCode(130);
 	}
 
 	// Fill out the Ipc details and publish.
