Version bump
eac9eec / Martin Ridgers / [REDACTED_PHONE]:58:25 +0200

diff --git a/Engine/Source/Programs/UnrealTraceServer/src/Version.h b/Engine/Source/Programs/UnrealTraceServer/src/Version.h
--- a/Engine/Source/Programs/UnrealTraceServer/src/Version.h
+++ b/Engine/Source/Programs/UnrealTraceServer/src/Version.h
@@ -7,7 +7,7 @@
 #endif
 
 #if !defined(TS_VERSION_MINOR)
-#	define TS_VERSION_MINOR		2
+#	define TS_VERSION_MINOR		3
 #endif
 
 /* vim: set noexpandtab : */
