Wrong version number was getting set
62100ce / Martin Ridgers / [REDACTED_PHONE]:05:46 +0200

diff --git a/Engine/Source/Programs/UnrealTraceServer/src/Main.cpp b/Engine/Source/Programs/UnrealTraceServer/src/Main.cpp
--- a/Engine/Source/Programs/UnrealTraceServer/src/Main.cpp
+++ b/Engine/Source/Programs/UnrealTraceServer/src/Main.cpp
@@ -121,7 +121,7 @@ public:
 ////////////////////////////////////////////////////////////////////////////////
 void FInstanceInfo::Set()
 {
-	Version = 1;
+	Version = CurrentVersion;
 #if TS_USING(TS_PLATFORM_WINDOWS)
 	Pid = GetCurrentProcessId();
 #elif TS_USING(TS_PLATFORM_LINUX) || TS_USING(TS_PLATFORM_MAC)
