Use a different default port for traces
5ea7e4c / Martin Ridgers / [REDACTED_PHONE]:26:52 +0200

diff --git a/Engine/Source/Programs/UnrealTraceServer/src/Main.cpp b/Engine/Source/Programs/UnrealTraceServer/src/Main.cpp
--- a/Engine/Source/Programs/UnrealTraceServer/src/Main.cpp
+++ b/Engine/Source/Programs/UnrealTraceServer/src/Main.cpp
@@ -92,7 +92,7 @@ static FStoreService* StartStore(const char* StoreDir)
 	FStoreService::FDesc Desc;
 	Desc.StoreDir = StoreDir;
 	Desc.StorePort = 1989;
-	Desc.RecorderPort = 1980;
+	Desc.RecorderPort = 1981;
 	return FStoreService::Create(Desc);
 }
 
