Expose the store server's icon as a named resource
e66ac21 / Martin Ridgers / [REDACTED_PHONE]:11:07 +0200

diff --git a/Engine/Source/Programs/UnrealTraceServer/src/Version.h b/Engine/Source/Programs/UnrealTraceServer/src/Version.h
--- a/Engine/Source/Programs/UnrealTraceServer/src/Version.h
+++ b/Engine/Source/Programs/UnrealTraceServer/src/Version.h
@@ -10,4 +10,6 @@
 #	define TS_VERSION_MINOR		4
 #endif
 
+#define TS_ICON_ID				493
+
 /* vim: set noexpandtab : */
diff --git a/Engine/Source/Programs/UnrealTraceServer/src/Version.rc b/Engine/Source/Programs/UnrealTraceServer/src/Version.rc
--- a/Engine/Source/Programs/UnrealTraceServer/src/Version.rc
+++ b/Engine/Source/Programs/UnrealTraceServer/src/Version.rc
@@ -19,6 +19,6 @@
 	}
 }
 
-2 ICON "Unreal.ico"
+TS_ICON_ID ICON "Unreal.ico"
 
 /* vim: set noexpandtab : */
