Fixed incorrect timestamp bias
3b2ac80 / Martin Ridgers / [REDACTED_PHONE]:00:29 +0200

diff --git a/Engine/Source/Programs/UnrealTraceServer/src/Store.cpp b/Engine/Source/Programs/UnrealTraceServer/src/Store.cpp
--- a/Engine/Source/Programs/UnrealTraceServer/src/Store.cpp
+++ b/Engine/Source/Programs/UnrealTraceServer/src/Store.cpp
@@ -3,6 +3,18 @@
 #include "Pch.h"
 #include "Store.h"
 
+////////////////////////////////////////////////////////////////////////////////
+// Pre-C++20 there is now way to convert between clocks. C++20 onwards it is
+// possible to create a clock with an Unreal epoch and convert file times to it.
+// But at time of writing, C++20 isn't complete enough across the board.
+static const uint64	UnrealEpochYear		= 1;
+#if TS_USING(TS_PLATFORM_WINDOWS)
+static const uint64	FsEpochYear			= 1601;
+#else
+static const uint64	FsEpochYear			= 1970;
+#endif
+static int64 FsToUnrealEpochBiasSeconds	= uint64(double(FsEpochYear - UnrealEpochYear) * 365.2425) * 86400;
+
 #if TS_USING(TS_WITH_DIR_WATCHER)
 ////////////////////////////////////////////////////////////////////////////////
 #if TS_USING(TS_PLATFORM_WINDOWS)
@@ -49,11 +61,12 @@ FStore::FTrace::FTrace(const char* InPath)
 
 	Id = QuickStoreHash(Name);
 
-	// Calculate that trace's timestamp
-	static const uint64 UnixToUnrealEpochBias = 0x089f7ff5f7b58000ull;
+	// Calculate that trace's timestamp. Bias in seconds then convert to 0.1us.
 	std::filesystem::file_time_type LastWriteTime = std::filesystem::last_write_time(InPath);
-	Timestamp = LastWriteTime.time_since_epoch().count();
-	Timestamp += UnixToUnrealEpochBias;
+	auto LastWriteDuration = LastWriteTime.time_since_epoch();
+	Timestamp = std::chrono::duration_cast<std::chrono::seconds>(LastWriteDuration).count();
+	Timestamp += FsToUnrealEpochBiasSeconds;
+	Timestamp *= 10'000'000;
 }
 
 ////////////////////////////////////////////////////////////////////////////////
