Simple strategy to trigger the shutdown flow in debug builds
dbcc228 / Martin Ridgers / [REDACTED_PHONE]:10:19 +0200

diff --git a/Engine/Source/Programs/UnrealTraceServer/src/Main.cpp b/Engine/Source/Programs/UnrealTraceServer/src/Main.cpp
--- a/Engine/Source/Programs/UnrealTraceServer/src/Main.cpp
+++ b/Engine/Source/Programs/UnrealTraceServer/src/Main.cpp
@@ -320,6 +320,19 @@ static int MainFork(int ArgC, char** ArgV)
 	}
 
 #if TS_USING(TS_BUILD_DEBUG)
+	static bool volatile bShouldExit = false;
+	while (true)
+	{
+		Sleep(500);
+
+		if (bShouldExit)
+		{
+			FWinHandle QuitEvent = CreateEventW(nullptr, TRUE, FALSE, GQuitEventName);
+			SetEvent(QuitEvent);
+			break;
+		}
+	}
+
 	DaemonThread.join();
 #else
 	CloseHandle(ProcessInfo.hProcess);
