Version bump
3585dcf / Martin Ridgers / [REDACTED_PHONE]:18:28 +0200

diff --git a/Engine/Source/Programs/UnrealTraceServer/src/Version.h b/Engine/Source/Programs/UnrealTraceServer/src/Version.h
--- a/Engine/Source/Programs/UnrealTraceServer/src/Version.h
+++ b/Engine/Source/Programs/UnrealTraceServer/src/Version.h
@@ -7,7 +7,7 @@
 #endif
 
 #if !defined(TS_VERSION_MINOR)
-#	define TS_VERSION_MINOR		4
+#	define TS_VERSION_MINOR		5
 #endif
 
 #define TS_ICON_ID				493
