Polyglot makefile
b18d259 / Martin Ridgers / [REDACTED_PHONE]:49:33 +0200

diff --git a/Engine/Source/Programs/UnrealTraceServer/makefile b/Engine/Source/Programs/UnrealTraceServer/makefile
--- a/Engine/Source/Programs/UnrealTraceServer/makefile
+++ b/Engine/Source/Programs/UnrealTraceServer/makefile
@@ -1,80 +1,17 @@
-# vim: set foldlevel=1 :
+# vim: expandtab foldlevel=1 ft=make
 
-# {{{1 debug and flags ---------------------------------------------------------
-OUT_DIR := _build
-ifeq ($(DEBUG), 1)
-CFLAGS := \
-	$(CFLAGS) \
-	-DTS_BUILD_DEBUG=TS_ON \
-	-Og
-OUT_DIR := $(OUT_DIR)/debug
-else
-CFLAGS := \
-	$(CFLAGS) \
-	-O3
-OUT_DIR := $(OUT_DIR)/development
-endif
-
-CFLAGS := \
-	$(CFLAGS) \
-	-g \
-	-Werror \
-	-std=c++17 \
-	-fno-exceptions \
-	-static
-
-UNAME := $(shell uname)
-ifeq ($(UNAME), Linux)
-	LDLIBS := $(LDLIBS) -lstdc++ -lpthread -lrt
-endif
-
-ifeq ($(UNAME), Darwin)
-	LDLIBS := $(LDLIBS) -lstdc++ -lpthread
-endif
-
-# {{{1 all/clean targets -------------------------------------------------------
-.SILENT :
+# Ensure "all" is the first encountered target so it is the default
+all :
 
-.PHONY : clean all unreal_trace
+ifdef MAKEDIR: # gmake=false condition, nmake=two unused targets
+!ifdef MAKEDIR # gmake=skipped,			nmake=true condition
 
-all : unreal_trace
+!include build/nmake.makefile
 
-clean :
-	rm -rf $(OUT_DIR)
-
-$(OUT_DIR) :
-	mkdir -p $@
-
-# {{{1 asio --------------------------------------------------------------------
-ASIO_VER = 1.12.2
-CFLAGS := \
-	$(CFLAGS) \
-	-I3rdparty/asio/$(ASIO_VER)
-
-# {{{1 store compile/link ------------------------------------------------------
-OBJS := \
-	$(OBJS) \
-	$(OUT_DIR)/Asio.o \
-	$(OUT_DIR)/AsioContext.o \
-	$(OUT_DIR)/AsioFile.o \
-	$(OUT_DIR)/AsioIoable.o \
-	$(OUT_DIR)/AsioSocket.o \
-	$(OUT_DIR)/AsioTcpServer.o \
-	$(OUT_DIR)/AsioTickable.o \
-	$(OUT_DIR)/Cbor.o \
-	$(OUT_DIR)/Main.o \
-	$(OUT_DIR)/Recorder.o \
-	$(OUT_DIR)/Store.o \
-	$(OUT_DIR)/StoreCborServer.o \
-	$(OUT_DIR)/StoreService.o \
-	$(OUT_DIR)/TraceRelay.o
-
-$(OUT_DIR)/%.o : src/%.cpp
-	echo $(<F)
-	$(CXX) $(CFLAGS) -c -o $@ $<
+!else
+else
 
-$(OUT_DIR)/UnrealTraceServer : $(OUT_DIR) $(OBJS)
-	echo $(@F)
-	$(CXX) -o $@ $(OBJS) $(LDLIBS)
+include build/gmake.makefile
 
-unreal_trace : $(OUT_DIR)/UnrealTraceServer
+endif 
+!endif :
