Changed to a trywait/sleep loop.
1b81493 / Martin Ridgers / [REDACTED_PHONE]:42:12 +0200

diff --git a/Engine/Source/Programs/UnrealTraceServer/src/Main.cpp b/Engine/Source/Programs/UnrealTraceServer/src/Main.cpp
--- a/Engine/Source/Programs/UnrealTraceServer/src/Main.cpp
+++ b/Engine/Source/Programs/UnrealTraceServer/src/Main.cpp
@@ -522,12 +522,19 @@ static int MainFork(int ArgC, char** ArgV)
 #endif // TS_BUILD_DEBUG
 
 	int Ret = 0;
-	timespec Timeout;
-	clock_gettime(CLOCK_REALTIME, &Timeout);
-	Timeout.tv_sec += (5000 / 1000); // expressed in milliseconds for grep-ability
-	if (sem_timedwait(BegunSem, &Timeout) < 0)
+	int TimeoutMs = 5000;
+	while (sem_trywait(BegunSem) < 0)
 	{
-		Ret = 8;
+		const uint32 SleepMs = 239;
+		timespec SleepTime = { 0, SleepMs * 1000 * 1000 };
+		nanosleep(&SleepTime, nullptr);
+
+		TimeoutMs -= SleepMs;
+		if (TimeoutMs <= 0)
+		{
+			Ret = 8;
+			break;
+		}
 	}
 
 #if TS_USING(TS_BUILD_DEBUG)
