More semantic return code handling
309257f / Martin Ridgers / [REDACTED_PHONE]:13:34 +0200

diff --git a/Engine/Source/Programs/UnrealTraceServer/src/Main.cpp b/Engine/Source/Programs/UnrealTraceServer/src/Main.cpp
--- a/Engine/Source/Programs/UnrealTraceServer/src/Main.cpp
+++ b/Engine/Source/Programs/UnrealTraceServer/src/Main.cpp
@@ -191,6 +191,22 @@ static int				MainDaemon(int, char**);
 void					AddToSystemTray(FStoreService&);
 void					RemoveFromSystemTray();
 
+////////////////////////////////////////////////////////////////////////////////
+enum : int
+{
+	Result_Ok				= 0,
+	Result_BegunExists,
+	Result_BegunTimeout,
+	Result_CopyFail,
+	Result_LaunchFail,
+	Result_NoQuitEvent,
+	Result_ProcessOpenFail,
+	Result_QuitExists,
+	Result_RenameFail,
+	Result_SharedMemFail,
+	Result_UnexpectedError,
+};
+
 ////////////////////////////////////////////////////////////////////////////////
 static int CreateExitCode(uint32 Id)
 {
@@ -211,7 +227,7 @@ static int MainFork(int ArgC, char** ArgV)
 		InstanceInfo->WaitForReady();
 		if (!InstanceInfo->IsOlder())
 		{
-			return 0;
+			return Result_Ok;
 		}
 
 		// Signal to the existing instance to shutdown or forcefully do it if it
@@ -220,14 +236,14 @@ static int MainFork(int ArgC, char** ArgV)
 		if (GetLastError() != ERROR_ALREADY_EXISTS)
 		{
 			// Assume someone else has closed the instance already.
-			return 0;
+			return Result_Ok;
 		}
 
 		DWORD OpenProcessFlags = PROCESS_TERMINATE | SYNCHRONIZE;
 		FWinHandle ProcHandle = OpenProcess(OpenProcessFlags, FALSE, InstanceInfo->Pid);
 		if (!ProcHandle.IsValid())
 		{
-			return CreateExitCode(1);
+			return CreateExitCode(Result_ProcessOpenFail);
 		}
 
 		SetEvent(QuitEvent);
@@ -244,7 +260,7 @@ static int MainFork(int ArgC, char** ArgV)
 	if (GetLastError() == ERROR_INSUFFICIENT_BUFFER)
 	{
 		// This should never really happen...
-		return CreateExitCode(3);
+		return CreateExitCode(Result_UnexpectedError);
 	}
 
 	// Calculate where to store the binaries.
@@ -275,7 +291,7 @@ static int MainFork(int ArgC, char** ArgV)
 		TempPath += Buffer;
 		if (!std::filesystem::copy_file(BinPath, TempPath))
 		{
-			return CreateExitCode(5);
+			return CreateExitCode(Result_CopyFail);
 		}
 
 		// Move the file into place. If this fails because the file exists then
@@ -285,7 +301,7 @@ static int MainFork(int ArgC, char** ArgV)
 		if (ErrorCode)
 		{
 			bool RaceLost = (ErrorCode == std::errc::file_exists);
-			return RaceLost ? 0 : CreateExitCode(7);
+			return RaceLost ? Result_Ok : CreateExitCode(Result_RenameFail);
 		}
 	}
 
@@ -293,7 +309,7 @@ static int MainFork(int ArgC, char** ArgV)
 	FWinHandle BegunEvent = CreateEventW(nullptr, TRUE, FALSE, GBegunEventName);
 	if (GetLastError() == ERROR_ALREADY_EXISTS)
 	{
-		return CreateExitCode(11);
+		return CreateExitCode(Result_BegunExists);
 	}
 
 	// For debugging ease and consistency we will daemonize in this process
@@ -309,14 +325,14 @@ static int MainFork(int ArgC, char** ArgV)
 
 	if (bOk == FALSE)
 	{
-		return CreateExitCode(13);
+		return CreateExitCode(Result_LaunchFail);
 	}
 #endif // TS_BUILD_DEBUG
 
-	int Ret = 0;
+	int Ret = Result_Ok;
 	if (WaitForSingleObject(BegunEvent, 5000) == WAIT_TIMEOUT)
 	{
-		Ret = CreateExitCode(17);
+		Ret = CreateExitCode(Result_BegunTimeout);
 	}
 
 #if TS_USING(TS_BUILD_DEBUG)
@@ -350,7 +366,7 @@ static int MainDaemon(int ArgC, char** ArgV)
 		PAGE_READWRITE, 0, GIpcSize, GIpcName);
 	if (!IpcHandle.IsValid())
 	{
-		return CreateExitCode(129);
+		return CreateExitCode(Result_SharedMemFail);
 	}
 
 	// Create a named event so others can tell us to quit.
@@ -359,7 +375,7 @@ static int MainDaemon(int ArgC, char** ArgV)
 	{
 		// This really should not happen. It is expected that only one process
 		// will get this far (gated by the shared-memory object creation).
-		return CreateExitCode(130);
+		return CreateExitCode(Result_QuitExists);
 	}
 
 	// Fill out the Ipc details and publish.
@@ -402,7 +418,7 @@ static int MainDaemon(int ArgC, char** ArgV)
 	// Clean up. We are done here.
 	RemoveFromSystemTray();
 	delete StoreService;
-	return 0;
+	return Result_Ok;
 }
 
 ////////////////////////////////////////////////////////////////////////////////
