Added an 'UnrealTraceServer kill' for Windows
c160a1f / Martin Ridgers / [REDACTED_PHONE]:13:52 +0200

diff --git a/Engine/Source/Programs/UnrealTraceServer/src/Main.cpp b/Engine/Source/Programs/UnrealTraceServer/src/Main.cpp
--- a/Engine/Source/Programs/UnrealTraceServer/src/Main.cpp
+++ b/Engine/Source/Programs/UnrealTraceServer/src/Main.cpp
@@ -213,6 +213,54 @@ static int CreateExitCode(uint32 Id)
 	return (GetLastError() & 0xfff) | (Id << 12);
 }
 
+////////////////////////////////////////////////////////////////////////////////
+static int MainKillImpl(int ArgC, char** ArgV, const FInstanceInfo* InstanceInfo)
+{
+	// Signal to the existing instance to shutdown or forcefully do it if it
+	// does not respond in time.
+
+	FWinHandle QuitEvent = CreateEventW(nullptr, TRUE, FALSE, GQuitEventName);
+	if (GetLastError() != ERROR_ALREADY_EXISTS)
+	{
+		return Result_NoQuitEvent;
+	}
+
+	DWORD OpenProcessFlags = PROCESS_TERMINATE | SYNCHRONIZE;
+	FWinHandle ProcHandle = OpenProcess(OpenProcessFlags, FALSE, InstanceInfo->Pid);
+	if (!ProcHandle.IsValid())
+	{
+		return CreateExitCode(Result_ProcessOpenFail);
+	}
+
+	SetEvent(QuitEvent);
+
+	if (WaitForSingleObject(ProcHandle, 5000) == WAIT_TIMEOUT)
+	{
+		TerminateProcess(ProcHandle, 10);
+	}
+
+	return Result_Ok;
+}
+
+////////////////////////////////////////////////////////////////////////////////
+static int MainKill(int ArgC, char** ArgV)
+{
+	// Find if an existing instance is already running.
+	FWinHandle IpcHandle = OpenFileMappingW(FILE_MAP_ALL_ACCESS, false, GIpcName);
+	if (!IpcHandle.IsValid())
+	{
+		return Result_Ok;
+	}
+
+	// There is an instance running so we can get its info block
+	void* IpcPtr = MapViewOfFile(IpcHandle, FILE_MAP_ALL_ACCESS, 0, 0, GIpcSize);
+	FMmapScope MmapScope(IpcPtr);
+	const auto* InstanceInfo = MmapScope.As<const FInstanceInfo>();
+	InstanceInfo->WaitForReady();
+
+	return MainKillImpl(ArgC, ArgV, InstanceInfo);
+}
+
 ////////////////////////////////////////////////////////////////////////////////
 static int MainFork(int ArgC, char** ArgV)
 {
@@ -230,28 +278,13 @@ static int MainFork(int ArgC, char** ArgV)
 			return Result_Ok;
 		}
 
-		// Signal to the existing instance to shutdown or forcefully do it if it
-		// does not respond in time.
-		FWinHandle QuitEvent = CreateEventW(nullptr, TRUE, FALSE, GQuitEventName);
-		if (GetLastError() != ERROR_ALREADY_EXISTS)
+		// Kill the other instance. If no quit event was found then we shall
+		// assume that another new store instance beat us to it.
+		int KillRet = MainKillImpl(0, nullptr, InstanceInfo);
+		if (KillRet == Result_NoQuitEvent || KillRet == Result_Ok)
 		{
-			// Assume someone else has closed the instance already.
 			return Result_Ok;
 		}
-
-		DWORD OpenProcessFlags = PROCESS_TERMINATE | SYNCHRONIZE;
-		FWinHandle ProcHandle = OpenProcess(OpenProcessFlags, FALSE, InstanceInfo->Pid);
-		if (!ProcHandle.IsValid())
-		{
-			return CreateExitCode(Result_ProcessOpenFail);
-		}
-
-		SetEvent(QuitEvent);
-
-		if (WaitForSingleObject(ProcHandle, 5000) == WAIT_TIMEOUT)
-		{
-			TerminateProcess(ProcHandle, 10);
-		}
 	}
 
 	// Get this binary's path
@@ -653,7 +686,7 @@ int main(int ArgC, char** ArgV)
 		"fork",		MainFork,
 		"daemon",	MainDaemon,
 		"test",		MainTest,
-		//"kill",		MainKill,
+		"kill",		MainKill,
 	};
 
 	if (ArgC > 1)
