Updated POSIX platforms to use the Result_* enum
776c1be / Martin Ridgers / [REDACTED_PHONE]:02:23 +0200

diff --git a/Engine/Source/Programs/UnrealTraceServer/src/Main.cpp b/Engine/Source/Programs/UnrealTraceServer/src/Main.cpp
--- a/Engine/Source/Programs/UnrealTraceServer/src/Main.cpp
+++ b/Engine/Source/Programs/UnrealTraceServer/src/Main.cpp
@@ -155,6 +155,29 @@ bool FInstanceInfo::IsOlder() const
 
 
 
+// {{{1 return codes -----------------------------------------------------------
+
+////////////////////////////////////////////////////////////////////////////////
+enum : int
+{
+	Result_Ok				= 0,
+	Result_BegunCreateFail,
+	Result_BegunExists,
+	Result_BegunTimeout,
+	Result_CopyFail,
+	Result_ForkFail,
+	Result_LaunchFail,
+	Result_NoQuitEvent,
+	Result_ProcessOpenFail,
+	Result_QuitExists,
+	Result_RenameFail,
+	Result_SharedMemFail,
+	Result_SharedMemTruncFail,
+	Result_UnexpectedError,
+};
+
+
+
 // {{{1 windows ----------------------------------------------------------------
 
 #if TS_USING(TS_PLATFORM_WINDOWS)
@@ -191,22 +214,6 @@ static int				MainDaemon(int, char**);
 void					AddToSystemTray(FStoreService&);
 void					RemoveFromSystemTray();
 
-////////////////////////////////////////////////////////////////////////////////
-enum : int
-{
-	Result_Ok				= 0,
-	Result_BegunExists,
-	Result_BegunTimeout,
-	Result_CopyFail,
-	Result_LaunchFail,
-	Result_NoQuitEvent,
-	Result_ProcessOpenFail,
-	Result_QuitExists,
-	Result_RenameFail,
-	Result_SharedMemFail,
-	Result_UnexpectedError,
-};
-
 ////////////////////////////////////////////////////////////////////////////////
 static int CreateExitCode(uint32 Id)
 {
@@ -514,7 +521,7 @@ static int MainFork(int ArgC, char** ArgV)
 		InstanceInfo->WaitForReady();
 		if (!InstanceInfo->IsOlder())
 		{
-			return 0;
+			return Result_Ok;
 		}
 
 		// Signal to the existing instance to shutdown.
@@ -522,7 +529,7 @@ static int MainFork(int ArgC, char** ArgV)
 		if (QuitSem == SEM_FAILED)
 		{
 			// Assume someone else has closed the instance already.
-			return 0;
+			return Result_Ok;
 		}
 		sem_post(QuitSem);
 		sem_close(QuitSem);
@@ -547,7 +554,7 @@ static int MainFork(int ArgC, char** ArgV)
 	sem_t* BegunSem = sem_open(GBegunSemName, O_RDONLY|O_CREAT|O_EXCL, 0644, 0);
 	if (BegunSem == SEM_FAILED)
 	{
-		return (errno == EEXIST) ? 7 : 11;
+		return (errno == EEXIST) ? Result_BegunExists : Result_BegunCreateFail;
 	}
 	OnScopeExit([=] () {
 		sem_unlink(GBegunSemName);
@@ -562,7 +569,7 @@ static int MainFork(int ArgC, char** ArgV)
 	pid_t DaemonPid = fork();
 	if (DaemonPid < 0)
 	{
-		return 3;
+		return Result_ForkFail;
 	}
 	else if (DaemonPid == 0)
 	{
@@ -570,7 +577,7 @@ static int MainFork(int ArgC, char** ArgV)
 	}
 #endif // TS_BUILD_DEBUG
 
-	int Ret = 0;
+	int Ret = Result_Ok;
 	int TimeoutMs = 5000;
 	while (sem_trywait(BegunSem) < 0)
 	{
@@ -581,7 +588,7 @@ static int MainFork(int ArgC, char** ArgV)
 		TimeoutMs -= SleepMs;
 		if (TimeoutMs <= 0)
 		{
-			Ret = 8;
+			Ret = Result_BegunTimeout;
 			break;
 		}
 	}
@@ -600,7 +607,7 @@ static int MainDaemon(int ArgC, char** ArgV)
 	int ShmHandle = shm_open(GShmName, O_RDWR|O_CREAT|O_EXCL, S_IROTH|S_IWOTH);
 	if (ShmHandle < 0)
 	{
-		return 1;
+		return Result_SharedMemFail;
 	}
 	OnScopeExit([=] () {
 		shm_unlink(GShmName);
@@ -613,7 +620,7 @@ static int MainDaemon(int ArgC, char** ArgV)
 	{
 		// This really should not happen. It is expected that only one process
 		// will get this far (gated by the shared-memory object creation).
-		return 2;
+		return Result_UnexpectedError;
 	}
 	OnScopeExit([=] () {
 		sem_unlink(GQuitSemName);
@@ -623,7 +630,7 @@ static int MainDaemon(int ArgC, char** ArgV)
 	// Fill out the Ipc details and publish.
 	if (ftruncate(ShmHandle, GShmSize) != 0)
 	{
-		return 3;
+		return Result_SharedMemTruncFail;
 	}
 
 	int32 ProtFlags = PROT_READ|PROT_WRITE;
@@ -659,7 +666,7 @@ static int MainDaemon(int ArgC, char** ArgV)
 
 	// Clean up. We are done here.
 	delete StoreService;
-	return 0;
+	return Result_Ok;
 }
 
 #endif // TS_PLATFORM_LINUX/MAC
