Fixed kill-the-instance races exiting too keenly
59a738c / Martin Ridgers / [REDACTED_PHONE]:03:01 +0200

diff --git a/Engine/Source/Programs/UnrealTraceServer/src/Main.cpp b/Engine/Source/Programs/UnrealTraceServer/src/Main.cpp
--- a/Engine/Source/Programs/UnrealTraceServer/src/Main.cpp
+++ b/Engine/Source/Programs/UnrealTraceServer/src/Main.cpp
@@ -285,13 +285,19 @@ static int MainFork(int ArgC, char** ArgV)
 			return Result_Ok;
 		}
 
-		// Kill the other instance. If no quit event was found then we shall
-		// assume that another new store instance beat us to it.
+		// Kill the other instance.
 		int KillRet = MainKillImpl(0, nullptr, InstanceInfo);
-		if (KillRet == Result_NoQuitEvent || KillRet == Result_Ok)
+		if (KillRet == Result_NoQuitEvent)
 		{
+			// If no quit event was found then we shall assume that another new
+			// store instance beat us to it.
 			return Result_Ok;
 		}
+
+		if (KillRet != Result_Ok)
+		{
+			return KillRet;
+		}
 	}
 
 	// Get this binary's path
