Corrected shared-memory file permissions
f1423cc / Martin Ridgers / [REDACTED_PHONE]:05:41 +0200

diff --git a/Engine/Source/Programs/UnrealTraceServer/src/Main.cpp b/Engine/Source/Programs/UnrealTraceServer/src/Main.cpp
--- a/Engine/Source/Programs/UnrealTraceServer/src/Main.cpp
+++ b/Engine/Source/Programs/UnrealTraceServer/src/Main.cpp
@@ -516,7 +516,7 @@ static int				MainDaemon(int, char**);
 static int MainFork(int ArgC, char** ArgV)
 {
 	// Check for an existing instance that is already running.
-	int ShmHandle = shm_open(GShmName, O_RDONLY, S_IROTH|S_IWOTH);
+	int ShmHandle = shm_open(GShmName, O_RDONLY, 0444);
 	if (ShmHandle >= 0)
 	{
 		void* ShmPtr = mmap(nullptr, GShmSize, PROT_READ, MAP_SHARED, ShmHandle, 0);
@@ -610,7 +610,7 @@ static int MainFork(int ArgC, char** ArgV)
 static int MainDaemon(int ArgC, char** ArgV)
 {
 	// Create a piece of shared memory so all store instances can communicate.
-	int ShmHandle = shm_open(GShmName, O_RDWR|O_CREAT|O_EXCL, S_IROTH|S_IWOTH);
+	int ShmHandle = shm_open(GShmName, O_RDWR|O_CREAT|O_EXCL, 0644);
 	if (ShmHandle < 0)
 	{
 		return Result_SharedMemFail;
