Signals can interrupt sig_wait()
75444fd / Martin Ridgers / [REDACTED_PHONE]:06:12 +0200

diff --git a/Engine/Source/Programs/UnrealTraceServer/src/Main.cpp b/Engine/Source/Programs/UnrealTraceServer/src/Main.cpp
--- a/Engine/Source/Programs/UnrealTraceServer/src/Main.cpp
+++ b/Engine/Source/Programs/UnrealTraceServer/src/Main.cpp
@@ -668,7 +668,11 @@ static int MainDaemon(int ArgC, char** ArgV)
 	}
 
 	// Wait to be told to resign.
-	sem_wait(QuitSem);
+	do
+	{
+		sem_wait(QuitSem);
+	}
+	while (errno == EINTR);
 
 	// Clean up. We are done here.
 	delete StoreService;
