Added 'tags' to .gitignore
a2e14ea / Martin Ridgers / [REDACTED_PHONE]:40:40 +0200

diff --git a/Engine/Source/Programs/UnrealTraceServer/.gitignore b/Engine/Source/Programs/UnrealTraceServer/.gitignore
--- a/Engine/Source/Programs/UnrealTraceServer/.gitignore
+++ b/Engine/Source/Programs/UnrealTraceServer/.gitignore
@@ -1 +1,2 @@
 _build
+tags
