Delay slightly before reading modified timestamps.
2f30438 / Martin Ridgers / [REDACTED_PHONE]:51:52 +0200

diff --git a/Engine/Source/Programs/UnrealTraceServer/src/Store.cpp b/Engine/Source/Programs/UnrealTraceServer/src/Store.cpp
--- a/Engine/Source/Programs/UnrealTraceServer/src/Store.cpp
+++ b/Engine/Source/Programs/UnrealTraceServer/src/Store.cpp
@@ -177,10 +177,24 @@ void FStore::WatchDir()
 		}
 
 #if TS_USING(TS_PLATFORM_WINDOWS)
-		FindNextChangeNotification(DirWatcher->native_handle());
-#endif
+		// Windows doesn't update modified timestamps in a timely fashion when
+		// copying files (or it could be Explorer that doesn't update it until
+		// later). This is a not-so-pretty "wait for a little bit" workaround.
+		auto* DelayTimer = new asio::steady_timer(IoContext);
+		DelayTimer->expires_after(std::chrono::seconds(2));
+		DelayTimer->async_wait([this, DelayTimer] (const asio::error_code& ErrorCode)
+		{
+			delete DelayTimer;
+
+			Refresh();
+
+			FindNextChangeNotification(DirWatcher->native_handle());
+			WatchDir();
+		});
+#else
 		Refresh();
 		WatchDir();
+#endif
 	});
 }
 #endif // TS_WITH_DIR_WATCHER
