Implemented 'UTS kill' on POSIX
1b835e5 / Martin Ridgers / [REDACTED_PHONE]:55:04 +0200

diff --git a/Engine/Source/Programs/UnrealTraceServer/src/Main.cpp b/Engine/Source/Programs/UnrealTraceServer/src/Main.cpp
--- a/Engine/Source/Programs/UnrealTraceServer/src/Main.cpp
+++ b/Engine/Source/Programs/UnrealTraceServer/src/Main.cpp
@@ -512,6 +512,63 @@ static const char*		GBegunSemName				= "/UnrealTraceSemBegun";
 static const char*		GQuitSemName				= "/UnrealTraceSemQuit";
 static int				MainDaemon(int, char**);
 
+////////////////////////////////////////////////////////////////////////////////
+static int MainKillImpl(int ArgC, char** ArgV, const FInstanceInfo* InstanceInfo)
+{
+	// Signal to the existing instance to shutdown or forcefully do it if it
+	// does not respond in time.
+
+	sem_t* QuitSem = sem_open(GQuitSemName, O_RDWR, 0666, 0);
+	if (QuitSem == SEM_FAILED)
+	{
+		// Assume someone else has closed the instance already.
+		return Result_NoQuitEvent;
+	}
+	sem_post(QuitSem);
+	sem_close(QuitSem);
+
+	// Wait for the process to end. If it takes too long, kill it.
+	const uint32 SleepMs = 47;
+	timespec SleepTime = { 0, SleepMs * 1000 * 1000 };
+	nanosleep(&SleepTime, nullptr);
+	for (uint32 i = 0; ; i += SleepMs)
+	{
+		if (i >= 5000) // "5000" for grep-ability
+		{
+			kill(InstanceInfo->Pid, SIGKILL);
+			break;
+		}
+
+		if (kill(InstanceInfo->Pid, 0) == ESRCH)
+		{
+			break;
+		}
+
+		nanosleep(&SleepTime, nullptr);
+	}
+
+	return Result_Ok;
+}
+
+////////////////////////////////////////////////////////////////////////////////
+static int MainKill(int ArgC, char** ArgV)
+{
+	// Check for an existing instance that is already running.
+	int ShmHandle = shm_open(GShmName, O_RDONLY, 0444);
+	if (ShmHandle < 0)
+	{
+		return (errno == ENOENT) ? Result_Ok : Result_SharedMemFail;
+	}
+
+	void* ShmPtr = mmap(nullptr, GShmSize, PROT_READ, MAP_SHARED, ShmHandle, 0);
+	FMmapScope MmapScope(ShmPtr, GShmSize);
+	const auto* InstanceInfo = MmapScope.As<const FInstanceInfo>();
+	InstanceInfo->WaitForReady();
+	close(ShmHandle);
+
+	return MainKillImpl(ArgC, ArgV, InstanceInfo);
+}
+
 ////////////////////////////////////////////////////////////////////////////////
 static int MainFork(int ArgC, char** ArgV)
 {
@@ -530,29 +587,18 @@ static int MainFork(int ArgC, char** ArgV)
 			return Result_Ok;
 		}
 
-		// Signal to the existing instance to shutdown.
-		sem_t* QuitSem = sem_open(GQuitSemName, O_RDWR, 0666, 0);
-		if (QuitSem == SEM_FAILED)
+		// Terminate the other instance.
+		int KillRet = MainKillImpl(0, nullptr, InstanceInfo);
+		if (KillRet == Result_NoQuitEvent)
 		{
-			// Assume someone else has closed the instance already.
+			// If no quit event was found then we shall assume that another new
+			// store instance beat us to it.
 			return Result_Ok;
 		}
-		sem_post(QuitSem);
-		sem_close(QuitSem);
 
-		// Wait for the process to end. If it takes too long, kill it.
-		for (uint32 i = 0; ; i += 1000)
+		if (KillRet != Result_Ok)
 		{
-			if (kill(InstanceInfo->Pid, 0) < 0)
-			{
-				break;
-			}
-
-			if (i >= 5000) // "5000" for grap-ability
-			{
-				kill(InstanceInfo->Pid, SIGTERM);
-				break;
-			}
+			return KillRet;
 		}
 	}
 
