.ico source plus simple build script
2e35bcf / Martin Ridgers / [REDACTED_PHONE]:47:06 +0200

diff --git a/Engine/Source/Programs/UnrealTraceServer/src/Icon/make_ico.py b/Engine/Source/Programs/UnrealTraceServer/src/Icon/make_ico.py
new file mode 100644
--- /dev/null
+++ b/Engine/Source/Programs/UnrealTraceServer/src/Icon/make_ico.py
@@ -0,0 +1,11 @@
+from PIL import Image
+
+images = []
+sizes = [255, 128, 64, 32, 48, 16]
+for size in sizes:
+    print(f"Loading ue_{size}.png")
+    image = Image.open(f"ue_{size}.png")
+    images.append(image)
+
+print("Creating ../Unreal.ico")
+images[0].save("../Unreal.ico", append_images=images)
