A sleep/sheduler yield is sufficient here
bcb5bec / Martin Ridgers / [REDACTED_PHONE]:23:21 +0200

diff --git a/Engine/Source/Programs/UnrealTraceServer/src/Main.cpp b/Engine/Source/Programs/UnrealTraceServer/src/Main.cpp
--- a/Engine/Source/Programs/UnrealTraceServer/src/Main.cpp
+++ b/Engine/Source/Programs/UnrealTraceServer/src/Main.cpp
@@ -12,13 +12,12 @@
 #if TS_USING(TS_PLATFORM_LINUX) || TS_USING(TS_PLATFORM_MAC)
 #	include <pwd.h>
 #	include <semaphore.h>
+#	include <sched.h>
 #	include <signal.h>
 #	include <sys/mman.h>
 #	include <unistd.h>
 #endif
 
-#include <immintrin.h>
-
 // {{{1 misc -------------------------------------------------------------------
 
 ////////////////////////////////////////////////////////////////////////////////
@@ -134,7 +133,11 @@ void FInstanceInfo::Set()
 void FInstanceInfo::WaitForReady() const
 {
 	// Spin until this instance info is published (by another process)
-	for (;; _mm_pause())
+#if TS_USING(TS_PLATFORM_WINDOWS)
+	for (;; Sleep(0))
+#else
+	for (;; sched_yield())
+#endif
 	{
 		if (Published.load(std::memory_order_acquire))
 		{
