Added TS_ARCH_* macros
481c494 / Martin Ridgers / [REDACTED_PHONE]:24:18 +0200

diff --git a/Engine/Source/Programs/UnrealTraceServer/src/Pch.h b/Engine/Source/Programs/UnrealTraceServer/src/Pch.h
--- a/Engine/Source/Programs/UnrealTraceServer/src/Pch.h
+++ b/Engine/Source/Programs/UnrealTraceServer/src/Pch.h
@@ -19,6 +19,9 @@
 #define TS_PLATFORM_LINUX		TS_OFF
 #define TS_PLATFORM_MAC			TS_OFF
 
+#define TS_ARCH_X64				TS_ON
+#define TS_ARCH_ARM64			TS_OFF
+
 #if defined(_MSC_VER)
 #	undef  TS_PLATFORM_WINDOWS
 #	define TS_PLATFORM_WINDOWS	TS_ON
@@ -32,6 +35,12 @@
 #if defined(__APPLE__)
 #	undef  TS_PLATFORM_MAC
 #	define TS_PLATFORM_MAC		TS_ON
+#	if defined(__arm64__)
+#		undef TS_ARCH_X64
+#		undef TS_ARCH_ARM64
+#		define TS_ARCH_X64		TS_OFF
+#		define TS_ARCH_ARM64	TS_ON
+#	endif
 #endif
 
 #define TS_ARRAY_COUNT(x)		(sizeof(x) / sizeof((x)[0]))
