Version bump
089aa90 / Martin Ridgers / [REDACTED_PHONE]:06:25 +0200

diff --git a/Engine/Source/Programs/UnrealTraceServer/src/Version.h b/Engine/Source/Programs/UnrealTraceServer/src/Version.h
--- a/Engine/Source/Programs/UnrealTraceServer/src/Version.h
+++ b/Engine/Source/Programs/UnrealTraceServer/src/Version.h
@@ -7,7 +7,7 @@
 #endif
 
 #if !defined(TS_VERSION_MINOR)
-#	define TS_VERSION_MINOR		5
+#	define TS_VERSION_MINOR		6
 #endif
 
 #define TS_ICON_ID				493
