Moved file system location determination out to a free function
a22b93b / Martin Ridgers / [REDACTED_PHONE]:36:09 +0200

diff --git a/Engine/Source/Programs/UnrealTraceServer/src/Main.cpp b/Engine/Source/Programs/UnrealTraceServer/src/Main.cpp
--- a/Engine/Source/Programs/UnrealTraceServer/src/Main.cpp
+++ b/Engine/Source/Programs/UnrealTraceServer/src/Main.cpp
@@ -82,6 +82,28 @@ struct TOnScopeExit
 
 
 
+////////////////////////////////////////////////////////////////////////////////
+static void GetUnrealTraceHome(std::filesystem::path& Out, bool Make=false)
+{
+#if TS_USING(TS_PLATFORM_WINDOWS)
+	wchar_t Buffer[MAX_PATH];
+	SHGetFolderPathW(nullptr, CSIDL_LOCAL_APPDATA, nullptr, SHGFP_TYPE_CURRENT, Buffer);
+	Out = Buffer;
+	Out /= "UnrealEngine/Common/UnrealTrace";
+#else
+	int UserId = getuid();
+	const passwd* Passwd = getpwuid(UserId);
+	Out = Passwd->pw_dir;
+	Out /= "UnrealEngine/UnrealTrace";
+#endif
+
+	if (Make)
+	{
+		std::filesystem::create_directories(Out);
+	}
+}
+
+
 
 // {{{1 store ------------------------------------------------------------------
 
@@ -314,12 +336,7 @@ static int MainFork(int ArgC, char** ArgV)
 
 	// Calculate where to store the binaries.
 	std::filesystem::path DestPath;
-	{
-		wchar_t Buffer[MAX_PATH];
-		SHGetFolderPathW(nullptr, CSIDL_LOCAL_APPDATA, nullptr, SHGFP_TYPE_CURRENT, Buffer);
-		DestPath = Buffer;
-	}
-	DestPath /= "UnrealEngine/Common/UnrealTrace";
+	GetUnrealTraceHome(DestPath);
 	{
 		wchar_t Buffer[64];
 		_snwprintf(Buffer, TS_ARRAY_COUNT(Buffer), L"Bin/%08x/UnrealTraceServer.exe", FInstanceInfo::CurrentVersion);
@@ -439,10 +456,8 @@ static int MainDaemon(int ArgC, char** ArgV)
 	FStoreService* StoreService;
 	{
 		std::filesystem::path StoreDir;
-		wchar_t Buffer[MAX_PATH];
-		SHGetFolderPathW(nullptr, CSIDL_LOCAL_APPDATA, nullptr, SHGFP_TYPE_CURRENT, Buffer);
-		StoreDir = Buffer;
-		StoreDir /= "UnrealEngine/Common/UnrealTrace/Store";
+		GetUnrealTraceHome(StoreDir);
+		StoreDir /= "Store";
 
 		std::u8string StoreDirUtf8 = StoreDir.u8string();
 		StoreService = StartStore((const char*)(StoreDirUtf8.c_str()));
@@ -699,11 +714,9 @@ static int MainDaemon(int ArgC, char** ArgV)
 	// Fire up the store
 	FStoreService* StoreService;
 	{
-		int UserId = getuid();
-		const passwd* Passwd = getpwuid(UserId);
-
-		std::filesystem::path StoreDir = Passwd->pw_dir;
-		StoreDir /= "UnrealEngine/UnrealTrace/Store";
+		std::filesystem::path StoreDir;
+		GetUnrealTraceHome(StoreDir);
+		StoreDir /= "Store";
 
 		StoreService = StartStore(StoreDir.c_str());
 	}
