Added some basic provisioning for logging
5820867 / Martin Ridgers / [REDACTED_PHONE]:37:11 +0200

diff --git a/Engine/Source/Programs/UnrealTraceServer/src/Main.cpp b/Engine/Source/Programs/UnrealTraceServer/src/Main.cpp
--- a/Engine/Source/Programs/UnrealTraceServer/src/Main.cpp
+++ b/Engine/Source/Programs/UnrealTraceServer/src/Main.cpp
@@ -105,6 +105,127 @@ static void GetUnrealTraceHome(std::filesystem::path& Out, bool Make=false)
 
 
 
+
+// {{{1 logging ----------------------------------------------------------------
+
+////////////////////////////////////////////////////////////////////////////////
+#define TS_LOG(Format, ...) \
+	do { FLogging::Log(Format "\n", ##__VA_ARGS__); } while (false)
+
+////////////////////////////////////////////////////////////////////////////////
+class FLogging
+{
+public:
+	static void			Initialize();
+	static void			Shutdown();
+	static void			Log(const char* Format, ...);
+
+private:
+						FLogging();
+						~FLogging();
+						FLogging(const FLogging&) = delete;
+						FLogging(FLogging&&) = default;
+	void				LogImpl(const char* String) const;
+	static FLogging*	Instance;
+	FILE*				File = nullptr;
+};
+
+////////////////////////////////////////////////////////////////////////////////
+FLogging* FLogging::Instance = nullptr;
+
+////////////////////////////////////////////////////////////////////////////////
+FLogging::FLogging()
+{
+	std::filesystem::path LogPath;
+
+#if TS_USING(TS_PLATFORM_WINDOWS)
+	GetUnrealTraceHome(LogPath, true);
+	LogPath /= "Server.log";
+	File = _wfopen(LogPath.c_str(), L"wb");
+#else
+#	if TS_USING(TS_PLATFORM_MAC)
+	int UserId = getuid();
+	const passwd* Passwd = getpwuid(UserId);
+	LogPath = Passwd->pw_dir;
+	LogPath /= "Library/Logs/Unreal Engine/UnrealTrace/Server.log";
+#	else
+	LogPath = "/var/log/UnrealTrace/Server.log";
+#	endif
+
+	std::filesystem::create_directories(LogPath.parent_path());
+	File = fopen(LogPath.c_str(), "wb");
+#endif
+}
+
+////////////////////////////////////////////////////////////////////////////////
+FLogging::~FLogging()
+{
+	if (File != nullptr)
+	{
+		fclose(File);
+	}
+}
+
+////////////////////////////////////////////////////////////////////////////////
+void FLogging::Initialize()
+{
+	if (Instance != nullptr)
+	{
+		return;
+	}
+
+	Instance = new FLogging();
+}
+
+////////////////////////////////////////////////////////////////////////////////
+void FLogging::Shutdown()
+{
+	if (Instance == nullptr)
+	{
+		return;
+	}
+
+	delete Instance;
+	Instance = nullptr;
+}
+
+////////////////////////////////////////////////////////////////////////////////
+void FLogging::LogImpl(const char* String) const
+{
+	if (File != nullptr)
+	{
+		fputs(String, File);
+	}
+
+	fputs(String, stdout);
+}
+
+////////////////////////////////////////////////////////////////////////////////
+void FLogging::Log(const char* Format, ...)
+{
+	va_list VaList;
+	va_start(VaList, Format);
+
+	char Buffer[320];
+	vsnprintf(Buffer, TS_ARRAY_COUNT(Buffer), Format, VaList);
+	Buffer[TS_ARRAY_COUNT(Buffer) - 1] = '\0';
+
+	Instance->LogImpl(Buffer);
+
+	va_end(VaList);
+}
+
+
+
+////////////////////////////////////////////////////////////////////////////////
+struct FLoggingScope
+{
+	FLoggingScope()		{ FLogging::Initialize(); }
+	~FLoggingScope()	{ FLogging::Shutdown(); }
+};
+
+
+
 // {{{1 store ------------------------------------------------------------------
 
 ////////////////////////////////////////////////////////////////////////////////
@@ -758,6 +879,8 @@ int MainTest(int ArgC, char** ArgV)
 ////////////////////////////////////////////////////////////////////////////////
 int main(int ArgC, char** ArgV)
 {
+	FLoggingScope LoggingScope;
+
 	struct {
 		const char*	Verb;
 		int			(*Entry)(int, char**);
