Added some boot logging
3a52a2a / Martin Ridgers / [REDACTED_PHONE]:37:52 +0200

diff --git a/Engine/Source/Programs/UnrealTraceServer/src/Main.cpp b/Engine/Source/Programs/UnrealTraceServer/src/Main.cpp
--- a/Engine/Source/Programs/UnrealTraceServer/src/Main.cpp
+++ b/Engine/Source/Programs/UnrealTraceServer/src/Main.cpp
@@ -372,23 +372,29 @@ static int MainKillImpl(int ArgC, char** ArgV, const FInstanceInfo* InstanceInfo
 	// Signal to the existing instance to shutdown or forcefully do it if it
 	// does not respond in time.
 
+	TS_LOG("Opening quit event");
 	FWinHandle QuitEvent = CreateEventW(nullptr, TRUE, FALSE, GQuitEventName);
 	if (GetLastError() != ERROR_ALREADY_EXISTS)
 	{
+		TS_LOG("Not found (gle=%d)", GetLastError());
 		return Result_NoQuitEvent;
 	}
 
+	TS_LOG("Open the process %d", InstanceInfo->Pid);
 	DWORD OpenProcessFlags = PROCESS_TERMINATE | SYNCHRONIZE;
 	FWinHandle ProcHandle = OpenProcess(OpenProcessFlags, FALSE, InstanceInfo->Pid);
 	if (!ProcHandle.IsValid())
 	{
+		TS_LOG("Unsuccessful (gle=%d)", GetLastError());
 		return CreateExitCode(Result_ProcessOpenFail);
 	}
 
+	TS_LOG("Firing quit event and waiting for process");
 	SetEvent(QuitEvent);
 
 	if (WaitForSingleObject(ProcHandle, 5000) == WAIT_TIMEOUT)
 	{
+		TS_LOG("Timeout. Force terminating");
 		TerminateProcess(ProcHandle, 10);
 	}
 
@@ -402,6 +408,7 @@ static int MainKill(int ArgC, char** ArgV)
 	FWinHandle IpcHandle = OpenFileMappingW(FILE_MAP_ALL_ACCESS, false, GIpcName);
 	if (!IpcHandle.IsValid())
 	{
+		TS_LOG("All good. There was no active UTS process");
 		return Result_Ok;
 	}
 
@@ -418,6 +425,7 @@ static int MainKill(int ArgC, char** ArgV)
 static int MainFork(int ArgC, char** ArgV)
 {
 	// Check for an existing instance that is already running.
+	TS_LOG("Opening exist instance's shared memory");
 	FWinHandle IpcHandle = OpenFileMappingW(FILE_MAP_ALL_ACCESS, false, GIpcName);
 	if (IpcHandle.IsValid())
 	{
@@ -428,6 +436,7 @@ static int MainFork(int ArgC, char** ArgV)
 		InstanceInfo->WaitForReady();
 		if (!InstanceInfo->IsOlder())
 		{
+			TS_LOG("Existing instance is the same age or newer");
 			return Result_Ok;
 		}
 
@@ -437,23 +446,32 @@ static int MainFork(int ArgC, char** ArgV)
 		{
 			// If no quit event was found then we shall assume that another new
 			// store instance beat us to it.
+			TS_LOG("Looks like someone else has already taken care of the upgrade");
 			return Result_Ok;
 		}
 
 		if (KillRet != Result_Ok)
 		{
+			TS_LOG("Kill attempt failed (ret=%d)", KillRet);
 			return KillRet;
 		}
 	}
+	else
+	{
+		TS_LOG("No existing process/shared memory found");
+	}
 
 	// Get this binary's path
+	TS_LOG("Getting binary path");
 	wchar_t BinPath[MAX_PATH];
 	uint32 BinPathLen = GetModuleFileNameW(nullptr, BinPath, TS_ARRAY_COUNT(BinPath));
 	if (GetLastError() == ERROR_INSUFFICIENT_BUFFER)
 	{
 		// This should never really happen...
+		TS_LOG("MAX_PATH is not enough");
 		return CreateExitCode(Result_UnexpectedError);
 	}
+	TS_LOG("Binary located at '%ls'", BinPath);
 
 	// Calculate where to store the binaries.
 	std::filesystem::path DestPath;
@@ -463,11 +481,14 @@ static int MainFork(int ArgC, char** ArgV)
 		_snwprintf(Buffer, TS_ARRAY_COUNT(Buffer), L"Bin/%08x/UnrealTraceServer.exe", FInstanceInfo::CurrentVersion);
 		DestPath /= Buffer;
 	}
+	TS_LOG("Run path '%ls'", DestPath.c_str());
 
 	// Copy the binary out to a location where it doesn't matter if the file
 	// gets locked by the OS.
 	if (!std::filesystem::is_regular_file(DestPath))
 	{
+		TS_LOG("Copying to run path");
+
 		std::filesystem::create_directories(DestPath.parent_path());
 
 		// Tag the destination with our PID and copy
@@ -478,6 +499,7 @@ static int MainFork(int ArgC, char** ArgV)
 		TempPath += Buffer;
 		if (!std::filesystem::copy_file(BinPath, TempPath))
 		{
+			TS_LOG("File copy failed (gle=%d)", GetLastError());
 			return CreateExitCode(Result_CopyFail);
 		}
 
@@ -488,14 +510,21 @@ static int MainFork(int ArgC, char** ArgV)
 		if (ErrorCode)
 		{
 			bool RaceLost = (ErrorCode == std::errc::file_exists);
+			TS_LOG("Rename to destination failed (bRaceLost=%c)", bRaceLost);
 			return RaceLost ? Result_Ok : CreateExitCode(Result_RenameFail);
 		}
 	}
+	else
+	{
+		TS_LOG("Already exists");
+	}
 
 	// Launch a new instance as a daemon and wait until we know it has started
+	TS_LOG("Creating begun event");
 	FWinHandle BegunEvent = CreateEventW(nullptr, TRUE, FALSE, GBegunEventName);
 	if (GetLastError() == ERROR_ALREADY_EXISTS)
 	{
+		TS_LOG("Did not work (gle=%d)", GetLastError());
 		return CreateExitCode(Result_BegunExists);
 	}
 
@@ -516,9 +545,11 @@ static int MainFork(int ArgC, char** ArgV)
 	}
 #endif // TS_BUILD_DEBUG
 
+	TS_LOG("Waiting on begun event");
 	int Ret = Result_Ok;
 	if (WaitForSingleObject(BegunEvent, 5000) == WAIT_TIMEOUT)
 	{
+		TS_LOG("Wait timed out (gle=%d)", GetLastError());
 		Ret = CreateExitCode(Result_BegunTimeout);
 	}
 
@@ -542,6 +573,7 @@ static int MainFork(int ArgC, char** ArgV)
 	CloseHandle(ProcessInfo.hThread);
 #endif
 
+	TS_LOG("Complete (ret=%d)", Ret);
 	return Ret;
 }
 
@@ -549,31 +581,37 @@ static int MainFork(int ArgC, char** ArgV)
 static int MainDaemon(int ArgC, char** ArgV)
 {
 	// Create a piece of shared memory so all store instances can communicate.
+	TS_LOG("Creating some shared memory");
 	FWinHandle IpcHandle = CreateFileMappingW(INVALID_HANDLE_VALUE, nullptr,
 		PAGE_READWRITE, 0, GIpcSize, GIpcName);
 	if (!IpcHandle.IsValid())
 	{
+		TS_LOG("Creation unsuccessful (gle=%d)", GetLastError());
 		return CreateExitCode(Result_SharedMemFail);
 	}
 
 	// Create a named event so others can tell us to quit.
+	TS_LOG("Creating a quit event");
 	FWinHandle QuitEvent = CreateEventW(nullptr, TRUE, FALSE, GQuitEventName);
 	if (GetLastError() == ERROR_ALREADY_EXISTS)
 	{
 		// This really should not happen. It is expected that only one process
 		// will get this far (gated by the shared-memory object creation).
+		TS_LOG("It unexpectedly exists already");
 		return CreateExitCode(Result_QuitExists);
 	}
 
 	// Fill out the Ipc details and publish.
 	void* IpcPtr = MapViewOfFile(IpcHandle, FILE_MAP_ALL_ACCESS, 0, 0, GIpcSize);
 	{
+		TS_LOG("Writing shared instance info");
 		FMmapScope MmapScope(IpcPtr);
 		auto* InstanceInfo = MmapScope.As<FInstanceInfo>();
 		InstanceInfo->Set();
 	}
 
 	// Fire up the store
+	TS_LOG("Starting the store");
 	FStoreService* StoreService;
 	{
 		std::filesystem::path StoreDir;
@@ -657,16 +695,19 @@ static int MainKillImpl(int ArgC, char** ArgV, const FInstanceInfo* InstanceInfo
 	// Signal to the existing instance to shutdown or forcefully do it if it
 	// does not respond in time.
 
+	TS_LOG("Opening quit semaphore");
 	sem_t* QuitSem = sem_open(GQuitSemName, O_RDWR, 0666, 0);
 	if (QuitSem == SEM_FAILED)
 	{
 		// Assume someone else has closed the instance already.
+		TS_LOG("Create failed (errno=%d)", errno);
 		return Result_NoQuitEvent;
 	}
 	sem_post(QuitSem);
 	sem_close(QuitSem);
 
 	// Wait for the process to end. If it takes too long, kill it.
+	TS_LOG("Waiting for pid %d", InstanceInfo->Pid);
 	const uint32 SleepMs = 47;
 	timespec SleepTime = { 0, SleepMs * 1000 * 1000 };
 	nanosleep(&SleepTime, nullptr);
@@ -675,11 +716,13 @@ static int MainKillImpl(int ArgC, char** ArgV, const FInstanceInfo* InstanceInfo
 		if (i >= 5000) // "5000" for grep-ability
 		{
 			kill(InstanceInfo->Pid, SIGKILL);
+			TS_LOG("Timed out. Sent SIGKILL instead (errno=%d)", errno);
 			break;
 		}
 
 		if (kill(InstanceInfo->Pid, 0) == ESRCH)
 		{
+			TS_LOG("Process no longer exists");
 			break;
 		}
 
@@ -696,7 +739,16 @@ static int MainKill(int ArgC, char** ArgV)
 	int ShmHandle = shm_open(GShmName, O_RDONLY, 0444);
 	if (ShmHandle < 0)
 	{
-		return (errno == ENOENT) ? Result_Ok : Result_SharedMemFail;
+		if (errno == ENOENT)
+		{
+			TS_LOG("All good. Ain't nuffin' running me ol' mucker.");
+			return Result_Ok;
+		}
+		else
+		{
+			TS_LOG("Unable to open shared memory (%s, errno=%d)", GShmName, errno);
+			return Result_SharedMemFail;
+		}
 	}
 
 	void* ShmPtr = mmap(nullptr, GShmSize, PROT_READ, MAP_SHARED, ShmHandle, 0);
@@ -712,6 +764,7 @@ static int MainKill(int ArgC, char** ArgV)
 static int MainFork(int ArgC, char** ArgV)
 {
 	// Check for an existing instance that is already running.
+	TS_LOG("Attempting to open shared memory %s", GShmName);
 	int ShmHandle = shm_open(GShmName, O_RDONLY, 0444);
 	if (ShmHandle >= 0)
 	{
@@ -723,29 +776,41 @@ static int MainFork(int ArgC, char** ArgV)
 		InstanceInfo->WaitForReady();
 		if (!InstanceInfo->IsOlder())
 		{
+			TS_LOG("Existing instance is the same age or newer");
 			return Result_Ok;
 		}
 
 		// Terminate the other instance.
+		TS_LOG("Killing an older instance that is already running");
 		int KillRet = MainKillImpl(0, nullptr, InstanceInfo);
 		if (KillRet == Result_NoQuitEvent)
 		{
 			// If no quit event was found then we shall assume that another new
 			// store instance beat us to it.
+			TS_LOG("Looks like someone else has already taken care of the upgrade");
 			return Result_Ok;
 		}
 
 		if (KillRet != Result_Ok)
 		{
+			TS_LOG("Kill attempt failed (ret=%d)", KillRet);
 			return KillRet;
 		}
 	}
+	else
+	{
+		TS_LOG("Unable to open shared memory (errno=%d)", errno);
+		TS_LOG("There probably is not an existing instance running");
+	}
 
 	// Launch a new instance as a daemon and wait until we know it has started
+	TS_LOG("Creating semaphore %s", GBegunSemName);
 	sem_t* BegunSem = sem_open(GBegunSemName, O_RDONLY|O_CREAT|O_EXCL, 0644, 0);
 	if (BegunSem == SEM_FAILED)
 	{
-		return (errno == EEXIST) ? Result_BegunExists : Result_BegunCreateFail;
+		int Ret = (errno == EEXIST) ? Result_BegunExists : Result_BegunCreateFail;
+		TS_LOG("Failed (errno=%d, ret=%d)", errno, Ret);
+		return Ret;
 	}
 	OnScopeExit([=] () {
 		sem_unlink(GBegunSemName);
@@ -754,12 +819,15 @@ static int MainFork(int ArgC, char** ArgV)
 
 	// For debugging ease and consistency we will daemonize in this process
 	// instead of spawning a second one.
+	pid_t DaemonPid = -1;
 #if TS_USING(TS_BUILD_DEBUG)
 	std::thread DaemonThread([] () { MainDaemon(0, nullptr); });
 #else
-	pid_t DaemonPid = fork();
+	TS_LOG("Forking process");
+	DaemonPid = fork();
 	if (DaemonPid < 0)
 	{
+		TS_LOG("Failed (errno=%d)", errno);
 		return Result_ForkFail;
 	}
 	else if (DaemonPid == 0)
@@ -768,6 +836,7 @@ static int MainFork(int ArgC, char** ArgV)
 	}
 #endif // TS_BUILD_DEBUG
 
+	TS_LOG("Waiting for begun semaphore from pid %d", DaemonPid);
 	int Ret = Result_Ok;
 	int TimeoutMs = 5000;
 	while (sem_trywait(BegunSem) < 0)
@@ -779,6 +848,7 @@ static int MainFork(int ArgC, char** ArgV)
 		TimeoutMs -= SleepMs;
 		if (TimeoutMs <= 0)
 		{
+			TS_LOG("Timed out");
 			Ret = Result_BegunTimeout;
 			break;
 		}
@@ -788,6 +858,7 @@ static int MainFork(int ArgC, char** ArgV)
 	DaemonThread.join();
 #endif
 
+	TS_LOG("Forked complete (ret=%d)", Ret);
 	return Ret;
 }
 
@@ -795,9 +866,11 @@ static int MainFork(int ArgC, char** ArgV)
 static int MainDaemon(int ArgC, char** ArgV)
 {
 	// Create a piece of shared memory so all store instances can communicate.
+	TS_LOG("Creating shared memory to store instance data (%s)", GShmName);
 	int ShmHandle = shm_open(GShmName, O_RDWR|O_CREAT|O_EXCL, 0644);
 	if (ShmHandle < 0)
 	{
+		TS_LOG("Well that was not at all successful (errno=%d)", errno);
 		return Result_SharedMemFail;
 	}
 	OnScopeExit([=] () {
@@ -806,11 +879,13 @@ static int MainDaemon(int ArgC, char** ArgV)
 	});
 
 	// Create a named semaphore so others can tell us to quit.
+	TS_LOG("Creating a master quit semaphore %s", GQuitSemName);
 	sem_t* QuitSem = sem_open(GQuitSemName, O_RDWR|O_CREAT|O_EXCL, 0666, 0);
 	if (QuitSem == SEM_FAILED)
 	{
 		// This really should not happen. It is expected that only one process
 		// will get this far (gated by the shared-memory object creation).
+		TS_LOG("Something went wrong (errno=%d)", errno);
 		return Result_UnexpectedError;
 	}
 	OnScopeExit([=] () {
@@ -819,8 +894,10 @@ static int MainDaemon(int ArgC, char** ArgV)
 	});
 
 	// Fill out the Ipc details and publish.
+	TS_LOG("Truncating shared memory");
 	if (ftruncate(ShmHandle, GShmSize) != 0)
 	{
+		TS_LOG("Oh bother (errno=%d)", errno);
 		return Result_SharedMemTruncFail;
 	}
 
@@ -835,6 +912,7 @@ static int MainDaemon(int ArgC, char** ArgV)
 	// Fire up the store
 	FStoreService* StoreService;
 	{
+		TS_LOG("Starting the store");
 		std::filesystem::path StoreDir;
 		GetUnrealTraceHome(StoreDir);
 		StoreDir /= "Store";
@@ -851,11 +929,13 @@ static int MainDaemon(int ArgC, char** ArgV)
 	}
 
 	// Wait to be told to resign.
+	TS_LOG("Waiting on the quit semaphore");
 	do
 	{
 		sem_wait(QuitSem);
 	}
 	while (errno == EINTR);
+	TS_LOG("Wait over (errno=%d)", errno);
 
 	// Clean up. We are done here.
 	delete StoreService;
