Added missing 'b' prefix
3df33d8 / Martin Ridgers / [REDACTED_PHONE]:20:58 +0200

diff --git a/Engine/Source/Programs/UnrealTraceServer/src/Main.cpp b/Engine/Source/Programs/UnrealTraceServer/src/Main.cpp
--- a/Engine/Source/Programs/UnrealTraceServer/src/Main.cpp
+++ b/Engine/Source/Programs/UnrealTraceServer/src/Main.cpp
@@ -509,9 +509,9 @@ static int MainFork(int ArgC, char** ArgV)
 		std::filesystem::rename(TempPath, DestPath, ErrorCode);
 		if (ErrorCode)
 		{
-			bool RaceLost = (ErrorCode == std::errc::file_exists);
+			bool bRaceLost = (ErrorCode == std::errc::file_exists);
 			TS_LOG("Rename to destination failed (bRaceLost=%c)", bRaceLost);
-			return RaceLost ? Result_Ok : CreateExitCode(Result_RenameFail);
+			return bRaceLost ? Result_Ok : CreateExitCode(Result_RenameFail);
 		}
 	}
 	else
