Added some simple log rotation
7b97967 / Martin Ridgers / [REDACTED_PHONE]:03:24 +0200

diff --git a/Engine/Source/Programs/UnrealTraceServer/src/Main.cpp b/Engine/Source/Programs/UnrealTraceServer/src/Main.cpp
--- a/Engine/Source/Programs/UnrealTraceServer/src/Main.cpp
+++ b/Engine/Source/Programs/UnrealTraceServer/src/Main.cpp
@@ -136,25 +136,72 @@ FLogging* FLogging::Instance = nullptr;
 ////////////////////////////////////////////////////////////////////////////////
 FLogging::FLogging()
 {
-	std::filesystem::path LogPath;
+	std::filesystem::path LogDir;
 
+	// Find where the logs should be written to. Make sure it exists.
 #if TS_USING(TS_PLATFORM_WINDOWS)
-	GetUnrealTraceHome(LogPath, true);
-	LogPath /= "Server.log";
-	File = _wfopen(LogPath.c_str(), L"wb");
+	GetUnrealTraceHome(LogDir, true);
 #else
 #	if TS_USING(TS_PLATFORM_MAC)
 	int UserId = getuid();
 	const passwd* Passwd = getpwuid(UserId);
-	LogPath = Passwd->pw_dir;
-	LogPath /= "Library/Logs/Unreal Engine/UnrealTrace/Server.log";
+	LogDir = Passwd->pw_dir;
+	LogDir /= "Library/Logs/Unreal Engine/UnrealTrace";
 #	else
-	LogPath = "/var/log/UnrealTrace/Server.log";
+	LogDir = "/var/log/UnrealTrace";
 #	endif
 
-	std::filesystem::create_directories(LogPath.parent_path());
-	File = fopen(LogPath.c_str(), "wb");
+	std::filesystem::create_directories(LogDir.parent_path());
+#endif // !TS_PLATFORM_WINDOWS
+
+	// Fetch all existing logs.
+	struct FExistingLog
+	{
+		std::filesystem::path	Path;
+		uint32					Index;
+
+		int32 operator < (const FExistingLog& Rhs) const
+		{
+			return Index < Rhs.Index;
+		}
+	};
+	std::vector<FExistingLog> ExistingLogs;
+	for (const auto& DirItem : std::filesystem::directory_iterator(LogDir))
+	{
+		int32 Index = -1;
+		std::string StemUtf8 = DirItem.path().stem().string();
+		sscanf(StemUtf8.c_str(), "Server_%d", &Index);
+		if (Index >= 0)
+		{
+			ExistingLogs.push_back({
+				DirItem.path(),
+				uint32(Index)
+			});
+		}
+	}
+
+	// Sort and try and tidy up old logs.
+	static int32 MaxLogs = 12; // plus one new one
+	std::sort(ExistingLogs.begin(), ExistingLogs.end());
+	for (int32 i = 0, n = int32(ExistingLogs.size() - MaxLogs); i < n; ++i)
+	{
+		std::error_code ErrorCode;
+		std::filesystem::remove(ExistingLogs[i].Path, ErrorCode);
+	}
+
+
+	// Open the log file (note; can race other instances)
+	uint32 LogIndex = ExistingLogs.empty() ? 0 : ExistingLogs.back().Index;
+	for (; LogIndex++, File == nullptr;)
+	{
+		std::filesystem::path LogPath = LogDir / std::format("Server_{}.log", LogIndex);
+#if TS_USING(TS_PLATFORM_WINDOWS)
+		File = _wfopen(LogPath.c_str(), L"wbxN");
+#else
+		File = fopen(LogPath.c_str(), "wbx");
 #endif
+	}
+	while (File == nullptr);
 }
 
 ////////////////////////////////////////////////////////////////////////////////
diff --git a/Engine/Source/Programs/UnrealTraceServer/src/Pch.h b/Engine/Source/Programs/UnrealTraceServer/src/Pch.h
--- a/Engine/Source/Programs/UnrealTraceServer/src/Pch.h
+++ b/Engine/Source/Programs/UnrealTraceServer/src/Pch.h
@@ -64,6 +64,7 @@
 #include <cstdint>
 #include <ctime>
 #include <filesystem>
+#include <format>
 #include <string>
 #include <string_view>
 #include <vector>
