Do not fallback to forking when running without arguments
a3f384d / Martin Ridgers / [REDACTED_PHONE]:09:19 +0200

diff --git a/Engine/Source/Programs/UnrealTraceServer/src/Main.cpp b/Engine/Source/Programs/UnrealTraceServer/src/Main.cpp
--- a/Engine/Source/Programs/UnrealTraceServer/src/Main.cpp
+++ b/Engine/Source/Programs/UnrealTraceServer/src/Main.cpp
@@ -1027,6 +1027,9 @@ int main(int ArgC, char** ArgV)
 				return (Dispatch.Entry)(ArgC - 1, ArgV + 1);
 			}
 		}
+
+		printf("Unknown command '%s'\n", ArgV[1]);
+		return 127;
 	}
 
 	return MainFork(ArgC, ArgV);
