Always require a command to be given
a46df10 / Martin Ridgers / [REDACTED_PHONE]:45:47 +0200

diff --git a/Engine/Source/Programs/UnrealTraceServer/src/Main.cpp b/Engine/Source/Programs/UnrealTraceServer/src/Main.cpp
--- a/Engine/Source/Programs/UnrealTraceServer/src/Main.cpp
+++ b/Engine/Source/Programs/UnrealTraceServer/src/Main.cpp
@@ -1006,7 +1006,10 @@ int MainTest(int ArgC, char** ArgV)
 ////////////////////////////////////////////////////////////////////////////////
 int main(int ArgC, char** ArgV)
 {
-	FLoggingScope LoggingScope;
+	if (ArgC < 2)
+	{
+		return 127;
+	}
 
 	struct {
 		const char*	Verb;
@@ -1018,21 +1021,17 @@ int main(int ArgC, char** ArgV)
 		"kill",		MainKill,
 	};
 
-	if (ArgC > 1)
+	for (const auto& Dispatch : Dispatches)
 	{
-		for (const auto& Dispatch : Dispatches)
+		if (strcmp(ArgV[1], Dispatch.Verb) == 0)
 		{
-			if (strcmp(ArgV[1], Dispatch.Verb) == 0)
-			{
-				return (Dispatch.Entry)(ArgC - 1, ArgV + 1);
-			}
+			FLoggingScope LoggingScope;
+			return (Dispatch.Entry)(ArgC - 1, ArgV + 1);
 		}
-
-		printf("Unknown command '%s'\n", ArgV[1]);
-		return 127;
 	}
 
-	return MainFork(ArgC, ArgV);
+	printf("Unknown command '%s'\n", ArgV[1]);
+	return 126;
 }
 
 /* vim: set noexpandtab foldlevel=1 : */
