Added some help text
353f1ab / Martin Ridgers / [REDACTED_PHONE]:46:05 +0200

diff --git a/Engine/Source/Programs/UnrealTraceServer/src/Main.cpp b/Engine/Source/Programs/UnrealTraceServer/src/Main.cpp
--- a/Engine/Source/Programs/UnrealTraceServer/src/Main.cpp
+++ b/Engine/Source/Programs/UnrealTraceServer/src/Main.cpp
@@ -1008,6 +1008,24 @@ int main(int ArgC, char** ArgV)
 {
 	if (ArgC < 2)
 	{
+		printf("UnrealTraceServer v%d.%d / Unreal Engine / Epic Games\n\n", TS_VERSION_PROTOCOL, TS_VERSION_MINOR);
+		puts(  "Usage; <cmd>");
+		puts(  "Commands;");
+		puts(  "  fork   Starts a background server, upgrading any existing instance");
+		puts(  "  daemon The mode that a background server runs in");
+		puts(  "  kill   Shuts down a currently running instance");
+		puts("");
+		puts(  "UnrealTraceServer acts as a hub between runtimes that are tracing performance");
+		puts(  "instrumentation and tools like Unreal Insights that consume and present that");
+		puts(  "data for analysis. TCP ports 1981 and 1989 are used, where the former receives");
+		puts(  "trace data, and the latter is used by tools to query the server's store.");
+
+		std::filesystem::path HomeDir;
+		GetUnrealTraceHome(HomeDir);
+		HomeDir.make_preferred();
+		std::string HomeDirU8 = HomeDir.string();
+		printf("\nStore path; %s\n", HomeDirU8.c_str());
+
 		return 127;
 	}
 
