Missing define if Clang was used on not-Apple platforms
35869a0 / Martin Ridgers / [REDACTED_PHONE]:50:24 +0200

diff --git a/Engine/Source/Programs/UnrealTraceServer/src/Asio.h b/Engine/Source/Programs/UnrealTraceServer/src/Asio.h
--- a/Engine/Source/Programs/UnrealTraceServer/src/Asio.h
+++ b/Engine/Source/Programs/UnrealTraceServer/src/Asio.h
@@ -9,6 +9,9 @@
 			_Pragma("GCC diagnostic ignored \"-Wdeprecated-declarations\"")
 #		define TS_ASIO_INCLUDE_END \
 			_Pragma("GCC diagnostic pop")
+#	else
+#		define TS_ASIO_INCLUDE_BEGIN
+#		define TS_ASIO_INCLUDE_END
 #	endif
 #else
 #	define TS_ASIO_INCLUDE_BEGIN
