Stop std::fs calls raising exceptions if files or directories exist
f6d44f8 / Martin Ridgers / [REDACTED_PHONE]:52:09 +0200

diff --git a/Engine/Source/Programs/UnrealTraceServer/src/Main.cpp b/Engine/Source/Programs/UnrealTraceServer/src/Main.cpp
--- a/Engine/Source/Programs/UnrealTraceServer/src/Main.cpp
+++ b/Engine/Source/Programs/UnrealTraceServer/src/Main.cpp
@@ -99,7 +99,8 @@ static void GetUnrealTraceHome(std::filesystem::path& Out, bool Make=false)
 
 	if (Make)
 	{
-		std::filesystem::create_directories(Out);
+		std::error_code ErrorCode;
+		std::filesystem::create_directories(Out, ErrorCode);
 	}
 }
 
@@ -151,7 +152,14 @@ FLogging::FLogging()
 	LogDir = "/var/log/UnrealTrace";
 #	endif
 
-	std::filesystem::create_directories(LogDir.parent_path());
+	std::error_code ErrorCode;
+	if (!std::filesystem::create_directories(LogDir, ErrorCode))
+	{
+		if (errno == EACCES)
+		{
+			return;
+		}
+	}
 #endif // !TS_PLATFORM_WINDOWS
 
 	// Fetch all existing logs.
