Only iterate the log directory if it exists
b980707 / Martin Ridgers / [REDACTED_PHONE]:52:59 +0200

diff --git a/Engine/Source/Programs/UnrealTraceServer/src/Main.cpp b/Engine/Source/Programs/UnrealTraceServer/src/Main.cpp
--- a/Engine/Source/Programs/UnrealTraceServer/src/Main.cpp
+++ b/Engine/Source/Programs/UnrealTraceServer/src/Main.cpp
@@ -174,17 +174,17 @@ FLogging::FLogging()
 		}
 	};
 	std::vector<FExistingLog> ExistingLogs;
-	for (const auto& DirItem : std::filesystem::directory_iterator(LogDir))
+	if (std::filesystem::is_directory(LogDir))
 	{
-		int32 Index = -1;
-		std::string StemUtf8 = DirItem.path().stem().string();
-		sscanf(StemUtf8.c_str(), "Server_%d", &Index);
-		if (Index >= 0)
+		for (const auto& DirItem : std::filesystem::directory_iterator(LogDir))
 		{
-			ExistingLogs.push_back({
-				DirItem.path(),
-				uint32(Index)
-			});
+			int32 Index = -1;
+			std::string StemUtf8 = DirItem.path().stem().string();
+			sscanf(StemUtf8.c_str(), "Server_%d", &Index);
+			if (Index >= 0)
+			{
+				ExistingLogs.push_back({DirItem.path(), uint32(Index)});
+			}
 		}
 	}
 
