Limit the number of attempts to create a log file
d6ce1fd / Martin Ridgers / [REDACTED_PHONE]:53:45 +0200

diff --git a/Engine/Source/Programs/UnrealTraceServer/src/Main.cpp b/Engine/Source/Programs/UnrealTraceServer/src/Main.cpp
--- a/Engine/Source/Programs/UnrealTraceServer/src/Main.cpp
+++ b/Engine/Source/Programs/UnrealTraceServer/src/Main.cpp
@@ -200,7 +200,7 @@ FLogging::FLogging()
 
 	// Open the log file (note; can race other instances)
 	uint32 LogIndex = ExistingLogs.empty() ? 0 : ExistingLogs.back().Index;
-	for (; LogIndex++, File == nullptr;)
+	for (uint32 n = LogIndex + 10; File == nullptr && LogIndex < n;)
 	{
 		std::filesystem::path LogPath = LogDir / std::format("Server_{}.log", LogIndex);
 #if TS_USING(TS_PLATFORM_WINDOWS)
