Clang doesn't have C++20's std::format yet
67aa8c5 / Martin Ridgers / [REDACTED_PHONE]:54:21 +0200

diff --git a/Engine/Source/Programs/UnrealTraceServer/src/Main.cpp b/Engine/Source/Programs/UnrealTraceServer/src/Main.cpp
--- a/Engine/Source/Programs/UnrealTraceServer/src/Main.cpp
+++ b/Engine/Source/Programs/UnrealTraceServer/src/Main.cpp
@@ -202,7 +202,11 @@ FLogging::FLogging()
 	uint32 LogIndex = ExistingLogs.empty() ? 0 : ExistingLogs.back().Index;
 	for (uint32 n = LogIndex + 10; File == nullptr && LogIndex < n;)
 	{
-		std::filesystem::path LogPath = LogDir / std::format("Server_{}.log", LogIndex);
+		++LogIndex;
+		char LogName[128];
+		snprintf(LogName, TS_ARRAY_COUNT(LogName), "Server_%d.log", LogIndex);
+		std::filesystem::path LogPath = LogDir / LogName;
+
 #if TS_USING(TS_PLATFORM_WINDOWS)
 		File = _wfopen(LogPath.c_str(), L"wbxN");
 #else
diff --git a/Engine/Source/Programs/UnrealTraceServer/src/Pch.h b/Engine/Source/Programs/UnrealTraceServer/src/Pch.h
--- a/Engine/Source/Programs/UnrealTraceServer/src/Pch.h
+++ b/Engine/Source/Programs/UnrealTraceServer/src/Pch.h
@@ -64,7 +64,6 @@
 #include <cstdint>
 #include <ctime>
 #include <filesystem>
-#include <format>
 #include <string>
 #include <string_view>
 #include <vector>
