Erroneous while statement
9c1fc1f / Martin Ridgers / [REDACTED_PHONE]:54:46 +0200

diff --git a/Engine/Source/Programs/UnrealTraceServer/src/Main.cpp b/Engine/Source/Programs/UnrealTraceServer/src/Main.cpp
--- a/Engine/Source/Programs/UnrealTraceServer/src/Main.cpp
+++ b/Engine/Source/Programs/UnrealTraceServer/src/Main.cpp
@@ -213,7 +213,6 @@ FLogging::FLogging()
 		File = fopen(LogPath.c_str(), "wbx");
 #endif
 	}
-	while (File == nullptr);
 }
 
 ////////////////////////////////////////////////////////////////////////////////
