Use the same "home" location for logs and lock files
ab6f020 / Martin Ridgers / [REDACTED_PHONE]:03:09 +0200

diff --git a/Engine/Source/Programs/UnrealTraceServer/src/Main.cpp b/Engine/Source/Programs/UnrealTraceServer/src/Main.cpp
--- a/Engine/Source/Programs/UnrealTraceServer/src/Main.cpp
+++ b/Engine/Source/Programs/UnrealTraceServer/src/Main.cpp
@@ -137,30 +137,9 @@ FLogging* FLogging::Instance = nullptr;
 ////////////////////////////////////////////////////////////////////////////////
 FLogging::FLogging()
 {
-	std::filesystem::path LogDir;
-
 	// Find where the logs should be written to. Make sure it exists.
-#if TS_USING(TS_PLATFORM_WINDOWS)
+	std::filesystem::path LogDir;
 	GetUnrealTraceHome(LogDir, true);
-#else
-#	if TS_USING(TS_PLATFORM_MAC)
-	int UserId = getuid();
-	const passwd* Passwd = getpwuid(UserId);
-	LogDir = Passwd->pw_dir;
-	LogDir /= "Library/Logs/Unreal Engine/UnrealTrace";
-#	else
-	LogDir = "/var/log/UnrealTrace";
-#	endif
-
-	std::error_code ErrorCode;
-	if (!std::filesystem::create_directories(LogDir, ErrorCode))
-	{
-		if (errno == EACCES)
-		{
-			return;
-		}
-	}
-#endif // !TS_PLATFORM_WINDOWS
 
 	// Fetch all existing logs.
 	struct FExistingLog
