Flush log file writes
8c5de4f / Martin Ridgers / [REDACTED_PHONE]:03:40 +0200

diff --git a/Engine/Source/Programs/UnrealTraceServer/src/Main.cpp b/Engine/Source/Programs/UnrealTraceServer/src/Main.cpp
--- a/Engine/Source/Programs/UnrealTraceServer/src/Main.cpp
+++ b/Engine/Source/Programs/UnrealTraceServer/src/Main.cpp
@@ -232,6 +232,7 @@ void FLogging::LogImpl(const char* String) const
 	if (File != nullptr)
 	{
 		fputs(String, File);
+		fflush(File);
 	}
 
 	fputs(String, stdout);
