Use a lock file instead of a shared-memory and semaphores as they orphan easily
d0ce7dd / Martin Ridgers / [REDACTED_PHONE]:30:14 +0200

diff --git a/Engine/Source/Programs/UnrealTraceServer/src/Main.cpp b/Engine/Source/Programs/UnrealTraceServer/src/Main.cpp
--- a/Engine/Source/Programs/UnrealTraceServer/src/Main.cpp
+++ b/Engine/Source/Programs/UnrealTraceServer/src/Main.cpp
@@ -357,6 +357,9 @@ enum : int
 	Result_RenameFail,
 	Result_SharedMemFail,
 	Result_SharedMemTruncFail,
+	Result_OpenFailPid,
+	Result_ReadFailPid,
+	Result_ReadFailCmdLine,
 	Result_UnexpectedError,
 };
 
@@ -721,28 +724,27 @@ int WinMain(HINSTANCE, HINSTANCE, LPSTR, int)
 #if TS_USING(TS_PLATFORM_LINUX) || TS_USING(TS_PLATFORM_MAC)
 
 ////////////////////////////////////////////////////////////////////////////////
-static const char*		GShmName					= "/UnrealTraceShm";
-static const int32		GShmSize					= 4 << 10;
-static const char*		GBegunSemName				= "/UnrealTraceSemBegun";
-static const char*		GQuitSemName				= "/UnrealTraceSemQuit";
-static int				MainDaemon(int, char**);
+static int MainDaemon(int, char**);
 
 ////////////////////////////////////////////////////////////////////////////////
-static int MainKillImpl(int ArgC, char** ArgV, const FInstanceInfo* InstanceInfo)
+static std::filesystem::path GetLockFilePath()
 {
-	// Signal to the existing instance to shutdown or forcefully do it if it
-	// does not respond in time.
+	std::filesystem::path Ret;
+	GetUnrealTraceHome(Ret, true);
+	Ret /= "UnrealTraceServer.pid";
+	return Ret;
+}
 
-	TS_LOG("Opening quit semaphore");
-	sem_t* QuitSem = sem_open(GQuitSemName, O_RDWR, 0666, 0);
-	if (QuitSem == SEM_FAILED)
+////////////////////////////////////////////////////////////////////////////////
+static int MainKillImpl(int ArgC, char** ArgV, const FInstanceInfo* InstanceInfo)
+{
+	// Issue the terminate signal
+	TS_LOG("Sending SIGTERM to %d", InstanceInfo->Pid);
+	if (kill(InstanceInfo->Pid, SIGTERM) < 0)
 	{
-		// Assume someone else has closed the instance already.
-		TS_LOG("Create failed (errno=%d)", errno);
-		return Result_NoQuitEvent;
+		TS_LOG("Failed to send SIGTERM");
+		return Result_SharedMemFail;
 	}
-	sem_post(QuitSem);
-	sem_close(QuitSem);
 
 	// Wait for the process to end. If it takes too long, kill it.
 	TS_LOG("Waiting for pid %d", InstanceInfo->Pid);
@@ -758,10 +760,13 @@ static int MainKillImpl(int ArgC, char** ArgV, const FInstanceInfo* InstanceInfo
 			break;
 		}
 
-		if (kill(InstanceInfo->Pid, 0) == ESRCH)
+		if (kill(InstanceInfo->Pid, 0) < 0)
 		{
-			TS_LOG("Process no longer exists");
-			break;
+			if (errno == ESRCH)
+			{
+				TS_LOG("Process no longer exists");
+				break;
+			}
 		}
 
 		nanosleep(&SleepTime, nullptr);
@@ -773,27 +778,32 @@ static int MainKillImpl(int ArgC, char** ArgV, const FInstanceInfo* InstanceInfo
 ////////////////////////////////////////////////////////////////////////////////
 static int MainKill(int ArgC, char** ArgV)
 {
-	// Check for an existing instance that is already running.
-	int ShmHandle = shm_open(GShmName, O_RDONLY, 0444);
-	if (ShmHandle < 0)
+	// Open the pid file to detect an existing instance
+	std::filesystem::path DotPidPath = GetLockFilePath();
+	TS_LOG("Checking for a '%s' lock file", DotPidPath.c_str());
+	int DotPidFd = open(DotPidPath.c_str(), O_RDONLY);
+	if (DotPidFd < 0)
 	{
 		if (errno == ENOENT)
 		{
 			TS_LOG("All good. Ain't nuffin' running me ol' mucker.");
 			return Result_Ok;
 		}
-		else
-		{
-			TS_LOG("Unable to open shared memory (%s, errno=%d)", GShmName, errno);
-			return Result_SharedMemFail;
-		}
+
+		TS_LOG("Unable to open lock file (%s, errno=%d)", DotPidPath.c_str(), errno);
+		return Result_OpenFailPid;
 	}
 
-	void* ShmPtr = mmap(nullptr, GShmSize, PROT_READ, MAP_SHARED, ShmHandle, 0);
-	FMmapScope MmapScope(ShmPtr, GShmSize);
-	const auto* InstanceInfo = MmapScope.As<const FInstanceInfo>();
-	InstanceInfo->WaitForReady();
-	close(ShmHandle);
+	// Get the instance info from the buffer
+	char DotPidBuffer[sizeof(FInstanceInfo)];
+	int Result = read(DotPidFd, DotPidBuffer, sizeof(DotPidBuffer));
+	close(DotPidFd);
+	if (Result < sizeof(FInstanceInfo))
+	{
+		TS_LOG("Failed to read the .pid lock file (errno=%d)", errno);
+		return Result_ReadFailPid;
+	}
+	const auto* InstanceInfo = (const FInstanceInfo*)DotPidBuffer;
 
 	return MainKillImpl(ArgC, ArgV, InstanceInfo);
 }
@@ -801,24 +811,56 @@ static int MainKill(int ArgC, char** ArgV)
 ////////////////////////////////////////////////////////////////////////////////
 static int MainFork(int ArgC, char** ArgV)
 {
-	// Check for an existing instance that is already running.
-	TS_LOG("Attempting to open shared memory %s", GShmName);
-	int ShmHandle = shm_open(GShmName, O_RDONLY, 0444);
-	if (ShmHandle >= 0)
+	// Open the pid file to detect an existing instance
+	std::filesystem::path DotPidPath = GetLockFilePath();
+	TS_LOG("Checking for a '%s' lock file", DotPidPath.c_str());
+	for (int DotPidFd = open(DotPidPath.c_str(), O_RDONLY); DotPidFd >= 0; )
 	{
-		void* ShmPtr = mmap(nullptr, GShmSize, PROT_READ, MAP_SHARED, ShmHandle, 0);
-		FMmapScope MmapScope(ShmPtr, GShmSize);
-		close(ShmHandle);
+		// Get the instance info from the buffer
+		char DotPidBuffer[sizeof(FInstanceInfo)];
+		int Result = read(DotPidFd, DotPidBuffer, sizeof(DotPidBuffer));
+		close(DotPidFd);
+		if (Result < sizeof(FInstanceInfo))
+		{
+			TS_LOG("Failed to read the .pid lock file (errno=%d)", errno);
+			return Result_ReadFailPid;
+		}
+		const auto* InstanceInfo = (const FInstanceInfo*)DotPidBuffer;
 
-		const auto* InstanceInfo = MmapScope.As<const FInstanceInfo>();
-		InstanceInfo->WaitForReady();
+		// Check the pid is valid and appears to be one of us
+		char CmdLinePath[320];
+		snprintf(CmdLinePath, TS_ARRAY_COUNT(CmdLinePath), "/proc/%d/cmdline", InstanceInfo->Pid);
+		int CmdLineFd = open(CmdLinePath, O_RDONLY);
+		if (CmdLineFd < 0)
+		{
+			TS_LOG("Process %d does not exist", InstanceInfo->Pid);
+			break;
+		}
+
+		Result = read(CmdLineFd, CmdLinePath, sizeof(CmdLinePath) - 1);
+		close(CmdLineFd);
+		if (Result <= 0)
+		{
+			TS_LOG("Unable to read 'cmdline' for process %d", InstanceInfo->Pid);
+			return Result_ReadFailCmdLine;
+		}
+
+		CmdLinePath[Result] = '\0';
+		if (strstr("UnrealTraceServer", CmdLinePath) == nullptr)
+		{
+			TS_LOG("Process %d is unrelated", InstanceInfo->Pid);
+			unlink(DotPidPath.c_str());
+			break;
+		}
+
+		// Old enough for this fine establishment?
 		if (!InstanceInfo->IsOlder())
 		{
 			TS_LOG("Existing instance is the same age or newer");
 			return Result_Ok;
 		}
 
-		// Terminate the other instance.
+		// If we've got this far then there's an instance running that is old
 		TS_LOG("Killing an older instance that is already running");
 		int KillRet = MainKillImpl(0, nullptr, InstanceInfo);
 		if (KillRet == Result_NoQuitEvent)
@@ -835,28 +877,14 @@ static int MainFork(int ArgC, char** ArgV)
 			return KillRet;
 		}
 	}
-	else
-	{
-		TS_LOG("Unable to open shared memory (errno=%d)", errno);
-		TS_LOG("There probably is not an existing instance running");
-	}
 
-	// Launch a new instance as a daemon and wait until we know it has started
-	TS_LOG("Creating semaphore %s", GBegunSemName);
-	sem_t* BegunSem = sem_open(GBegunSemName, O_RDONLY|O_CREAT|O_EXCL, 0644, 0);
-	if (BegunSem == SEM_FAILED)
-	{
-		int Ret = (errno == EEXIST) ? Result_BegunExists : Result_BegunCreateFail;
-		TS_LOG("Failed (errno=%d, ret=%d)", errno, Ret);
-		return Ret;
-	}
-	OnScopeExit([=] () {
-		sem_unlink(GBegunSemName);
-		sem_close(BegunSem);
-	});
+	// Daemon mode expects there to be no lock file on disk
+	std::error_code ErrorCode;
+	std::filesystem::remove(DotPidPath, ErrorCode);
 
-	// For debugging ease and consistency we will daemonize in this process
-	// instead of spawning a second one.
+	// Launch a daemonized version of ourselves. For debugging ease and
+	// consistency we will daemonize in this process instead of spawning
+	// a second one.
 	pid_t DaemonPid = -1;
 #if TS_USING(TS_BUILD_DEBUG)
 	std::thread DaemonThread([] () { MainDaemon(0, nullptr); });
@@ -874,12 +902,24 @@ static int MainFork(int ArgC, char** ArgV)
 	}
 #endif // TS_BUILD_DEBUG
 
-	TS_LOG("Waiting for begun semaphore from pid %d", DaemonPid);
+	// Wait for the daemon to indicate that it has started the store.
 	int Ret = Result_Ok;
+	TS_LOG("Wait until we know the daemon has started.");
 	int TimeoutMs = 5000;
-	while (sem_trywait(BegunSem) < 0)
+	while (true)
 	{
-		const uint32 SleepMs = 239;
+		// Check lock file's size
+		struct stat DotPidStat;
+		if (stat(DotPidPath.c_str(), &DotPidStat) == 0)
+		{
+			if (DotPidStat.st_size > 0)
+			{
+				TS_LOG("Successful start detected. Yay!");
+				break;
+			}
+		}
+
+		const uint32 SleepMs = 67;
 		timespec SleepTime = { 0, SleepMs * 1000 * 1000 };
 		nanosleep(&SleepTime, nullptr);
 
@@ -903,49 +943,28 @@ static int MainFork(int ArgC, char** ArgV)
 ////////////////////////////////////////////////////////////////////////////////
 static int MainDaemon(int ArgC, char** ArgV)
 {
-	// Create a piece of shared memory so all store instances can communicate.
-	TS_LOG("Creating shared memory to store instance data (%s)", GShmName);
-	int ShmHandle = shm_open(GShmName, O_RDWR|O_CREAT|O_EXCL, 0644);
-	if (ShmHandle < 0)
+	// We expect that there is no lock file on disk if we've got this far.
+	std::filesystem::path DotPidPath = GetLockFilePath();
+	TS_LOG("Claiming lock file '%s'", DotPidPath.c_str());
+	int DotPidFd = open(DotPidPath.c_str(), O_CREAT|O_EXCL|O_WRONLY, 0666);
+	if (DotPidFd < 0)
 	{
-		TS_LOG("Well that was not at all successful (errno=%d)", errno);
-		return Result_SharedMemFail;
-	}
-	OnScopeExit([=] () {
-		shm_unlink(GShmName);
-		close(ShmHandle);
-	});
-
-	// Create a named semaphore so others can tell us to quit.
-	TS_LOG("Creating a master quit semaphore %s", GQuitSemName);
-	sem_t* QuitSem = sem_open(GQuitSemName, O_RDWR|O_CREAT|O_EXCL, 0666, 0);
-	if (QuitSem == SEM_FAILED)
-	{
-		// This really should not happen. It is expected that only one process
-		// will get this far (gated by the shared-memory object creation).
-		TS_LOG("Something went wrong (errno=%d)", errno);
-		return Result_UnexpectedError;
-	}
-	OnScopeExit([=] () {
-		sem_unlink(GQuitSemName);
-		sem_close(QuitSem);
-	});
+		if (errno == EEXIST)
+		{
+			TS_LOG("Lock file already exists");
+			return Result_OpenFailPid;
+		}
 
-	// Fill out the Ipc details and publish.
-	TS_LOG("Truncating shared memory");
-	if (ftruncate(ShmHandle, GShmSize) != 0)
-	{
-		TS_LOG("Oh bother (errno=%d)", errno);
-		return Result_SharedMemTruncFail;
+		TS_LOG("Unexpected error (errno=%d)", errno);
+		return Result_UnexpectedError;
 	}
 
-	int32 ProtFlags = PROT_READ|PROT_WRITE;
-	void* ShmPtr = mmap(nullptr, GShmSize, ProtFlags, MAP_SHARED, ShmHandle, 0);
-	{
-		FMmapScope MmapScope(ShmPtr, GShmSize);
-		auto* InstanceInfo = MmapScope.As<FInstanceInfo>();
-		InstanceInfo->Set();
-	}
+	// Block all signals on all threads
+	sigset_t SignalSet;
+	sigemptyset(&SignalSet);
+	sigaddset(&SignalSet, SIGTERM);
+	sigaddset(&SignalSet, SIGINT);
+	pthread_sigmask(SIG_BLOCK, &SignalSet, nullptr);
 
 	// Fire up the store
 	FStoreService* StoreService;
@@ -957,26 +976,34 @@ static int MainDaemon(int ArgC, char** ArgV)
 
 		StoreService = StartStore(StoreDir.c_str());
 	}
+	OnScopeExit([StoreService] () { delete StoreService; });
 
 	// Let every one know we've started.
-	sem_t* BegunSem = sem_open(GBegunSemName, O_RDWR, 0644, 0);
-	if (BegunSem != SEM_FAILED)
+	FInstanceInfo InstanceInfo;
+	InstanceInfo.Set();
+	if (write(DotPidFd, &InstanceInfo, sizeof(InstanceInfo)) != sizeof(InstanceInfo))
 	{
-		sem_post(BegunSem);
-		sem_close(BegunSem);
+		TS_LOG("Unable to write instance info to lock file (errno=%d)", errno);
+		return Result_UnexpectedError;
 	}
+	fsync(DotPidFd);
 
 	// Wait to be told to resign.
-	TS_LOG("Waiting on the quit semaphore");
-	do
+	TS_LOG("Entering signal wait loop...");
+	while (true)
 	{
-		sem_wait(QuitSem);
+		int Signal = -1;
+		int Ret = sigwait(&SignalSet, &Signal);
+		if (Ret == 0)
+		{
+			TS_LOG("Received signal %d", Signal);
+			break;
+		}
 	}
-	while (errno == EINTR);
-	TS_LOG("Wait over (errno=%d)", errno);
 
 	// Clean up. We are done here.
-	delete StoreService;
+	std::error_code ErrorCode;
+	std::filesystem::remove(DotPidPath, ErrorCode);
 	return Result_Ok;
 }
 
