UnrealTraceStore was locking the directory it was launched from
229bc59 / Martin Ridgers / [REDACTED_PHONE]:01:30 +0200

diff --git a/Engine/Source/Programs/UnrealTraceServer/src/Main.cpp b/Engine/Source/Programs/UnrealTraceServer/src/Main.cpp
--- a/Engine/Source/Programs/UnrealTraceServer/src/Main.cpp
+++ b/Engine/Source/Programs/UnrealTraceServer/src/Main.cpp
@@ -621,6 +621,24 @@ static int MainFork(int ArgC, char** ArgV)
 ////////////////////////////////////////////////////////////////////////////////
 static int MainDaemon(int ArgC, char** ArgV)
 {
+	// Move the working directory to be where this binary is located.
+	TS_LOG("Setting working directory");
+	wchar_t BinPath[MAX_PATH];
+	uint32 BinPathLen = GetModuleFileNameW(nullptr, BinPath, TS_ARRAY_COUNT(BinPath));
+	if (GetLastError() != ERROR_INSUFFICIENT_BUFFER && BinPathLen > 0)
+	{
+		std::error_code ErrorCode;
+		std::filesystem::path BinDir(BinPath);
+		BinDir = BinDir.parent_path();
+		std::filesystem::current_path(BinDir, ErrorCode);
+		const char* Result = ErrorCode ? "Failed" : "Succeeded";
+		TS_LOG("%s setting '%ls'", Result, BinDir.c_str());
+	}
+	else
+	{
+		TS_LOG("Something went wrong (gle=%d)", GetLastError());
+	}
+
 	// Create a piece of shared memory so all store instances can communicate.
 	TS_LOG("Creating some shared memory");
 	FWinHandle IpcHandle = CreateFileMappingW(INVALID_HANDLE_VALUE, nullptr,
