Put the threaded mode of the daemon behind its own define
b0ad3b6 / Martin Ridgers / [REDACTED_PHONE]:26:33 +0200

diff --git a/Engine/Source/Programs/UnrealTraceServer/src/Main.cpp b/Engine/Source/Programs/UnrealTraceServer/src/Main.cpp
--- a/Engine/Source/Programs/UnrealTraceServer/src/Main.cpp
+++ b/Engine/Source/Programs/UnrealTraceServer/src/Main.cpp
@@ -18,6 +18,14 @@
 #	include <unistd.h>
 #endif
 
+// Debug builds act as both the forker and the daemon. Set to TS_OFF to disable
+// this behaviour.
+#if TS_USING(TS_BUILD_DEBUG) //&&0
+#	define TS_DAEMON_THREAD TS_ON
+#else
+#	define TS_DAEMON_THREAD TS_OFF
+#endif
+
 // {{{1 misc -------------------------------------------------------------------
 
 ////////////////////////////////////////////////////////////////////////////////
@@ -571,7 +579,7 @@ static int MainFork(int ArgC, char** ArgV)
 
 	// For debugging ease and consistency we will daemonize in this process
 	// instead of spawning a second one.
-#if TS_USING(TS_BUILD_DEBUG)
+#if TS_USING(TS_DAEMON_THREAD)
 	std::thread DaemonThread([] () { MainDaemon(0, nullptr); });
 #else
 	uint32 Flags = CREATE_BREAKAWAY_FROM_JOB;
@@ -584,7 +592,7 @@ static int MainFork(int ArgC, char** ArgV)
 	{
 		return CreateExitCode(Result_LaunchFail);
 	}
-#endif // TS_BUILD_DEBUG
+#endif // !TS_DAEMON_THREAD
 
 	TS_LOG("Waiting on begun event");
 	int Ret = Result_Ok;
@@ -594,7 +602,7 @@ static int MainFork(int ArgC, char** ArgV)
 		Ret = CreateExitCode(Result_BegunTimeout);
 	}
 
-#if TS_USING(TS_BUILD_DEBUG)
+#if TS_USING(TS_DAEMON_THREAD)
 	static bool volatile bShouldExit = false;
 	while (true)
 	{
