Debug builds always copy the trace server's binary out of the tree
6f913e7 / Martin Ridgers / [REDACTED_PHONE]:27:18 +0200

diff --git a/Engine/Source/Programs/UnrealTraceServer/src/Main.cpp b/Engine/Source/Programs/UnrealTraceServer/src/Main.cpp
--- a/Engine/Source/Programs/UnrealTraceServer/src/Main.cpp
+++ b/Engine/Source/Programs/UnrealTraceServer/src/Main.cpp
@@ -483,7 +483,11 @@ static int MainFork(int ArgC, char** ArgV)
 
 		const auto* InstanceInfo = MmapScope.As<const FInstanceInfo>();
 		InstanceInfo->WaitForReady();
+#if TS_USING(TS_BUILD_DEBUG)
+		if (false)
+#else
 		if (!InstanceInfo->IsOlder())
+#endif
 		{
 			TS_LOG("Existing instance is the same age or newer");
 			return Result_Ok;
@@ -532,6 +536,14 @@ static int MainFork(int ArgC, char** ArgV)
 	}
 	TS_LOG("Run path '%ls'", DestPath.c_str());
 
+#if TS_USING(TS_BUILD_DEBUG)
+	// Debug builds will always do the copy.
+	{
+		std::error_code ErrorCode;
+		std::filesystem::remove(DestPath, ErrorCode);
+	}
+#endif
+
 	// Copy the binary out to a location where it doesn't matter if the file
 	// gets locked by the OS.
 	if (!std::filesystem::is_regular_file(DestPath))
