Use the error_code variant in case create_directories throws an exception
88a366f / Martin Ridgers / [REDACTED_PHONE]:27:55 +0200

diff --git a/Engine/Source/Programs/UnrealTraceServer/src/Main.cpp b/Engine/Source/Programs/UnrealTraceServer/src/Main.cpp
--- a/Engine/Source/Programs/UnrealTraceServer/src/Main.cpp
+++ b/Engine/Source/Programs/UnrealTraceServer/src/Main.cpp
@@ -550,7 +550,8 @@ static int MainFork(int ArgC, char** ArgV)
 	{
 		TS_LOG("Copying to run path");
 
-		std::filesystem::create_directories(DestPath.parent_path());
+		std::error_code ErrorCode;
+		std::filesystem::create_directories(DestPath.parent_path(), ErrorCode);
 
 		// Tag the destination with our PID and copy
 		DWORD OurPid = GetCurrentProcessId();
@@ -566,7 +567,6 @@ static int MainFork(int ArgC, char** ArgV)
 
 		// Move the file into place. If this fails because the file exists then
 		// another instance has beaten us to the punch.
-		std::error_code ErrorCode;
 		std::filesystem::rename(TempPath, DestPath, ErrorCode);
 		if (ErrorCode)
 		{
