Remove any administrator privileges from the daemon. Not that the token is still considered elevated for now.
ff3b45c / Martin Ridgers / [REDACTED_PHONE]:28:52 +0200

diff --git a/Engine/Source/Programs/UnrealTraceServer/src/Main.cpp b/Engine/Source/Programs/UnrealTraceServer/src/Main.cpp
--- a/Engine/Source/Programs/UnrealTraceServer/src/Main.cpp
+++ b/Engine/Source/Programs/UnrealTraceServer/src/Main.cpp
@@ -594,17 +594,58 @@ static int MainFork(int ArgC, char** ArgV)
 #if TS_USING(TS_DAEMON_THREAD)
 	std::thread DaemonThread([] () { MainDaemon(0, nullptr); });
 #else
-	uint32 Flags = CREATE_BREAKAWAY_FROM_JOB;
+	uint32 CreateProcFlags = CREATE_BREAKAWAY_FROM_JOB;
 	STARTUPINFOW StartupInfo = { sizeof(STARTUPINFOW) };
 	PROCESS_INFORMATION ProcessInfo = {};
-	BOOL bOk = CreateProcessW(DestPath.c_str(), LPWSTR(L"UnrealTraceServer.exe daemon"), nullptr,
-		nullptr, FALSE, Flags, nullptr, nullptr, &StartupInfo, &ProcessInfo);
 
-	if (bOk == FALSE)
+	// Limit the authority of the daemon.
+	SAFER_LEVEL_HANDLE SafeLevel = nullptr;
+	BOOL bCanLaunchSafely = SaferCreateLevel(SAFER_SCOPEID_USER,
+		SAFER_LEVELID_NORMALUSER, SAFER_LEVEL_OPEN, &SafeLevel, nullptr);
+	if (bCanLaunchSafely == TRUE)
 	{
-		return CreateExitCode(Result_LaunchFail);
+		HANDLE AccessToken;
+		if (SaferComputeTokenFromLevel(SafeLevel, nullptr, &AccessToken, 0, nullptr))
+		{
+			BOOL bOk = CreateProcessAsUserW(
+				AccessToken,
+				DestPath.c_str(),
+				LPWSTR(L"UnrealTraceServer.exe daemon"),
+				nullptr, nullptr, FALSE, CreateProcFlags, nullptr, nullptr,
+				&StartupInfo, &ProcessInfo);
+
+			if (bOk == FALSE)
+			{
+				bCanLaunchSafely = bOk;
+			}
+
+			CloseHandle(AccessToken);
+		}
+
+		SaferCloseLevel(SafeLevel);
 	}
-#endif // !TS_DAEMON_THREAD
+
+	// Fallback to a normal CreateProc() call if using a limited token failed
+	if (bCanLaunchSafely == FALSE)
+	{
+		BOOL bOk = CreateProcessW(
+			DestPath.c_str(),
+			LPWSTR(L"UnrealTraceServer.exe daemon"),
+			nullptr, nullptr, FALSE, CreateProcFlags, nullptr, nullptr,
+			&StartupInfo, &ProcessInfo);
+
+		if (bOk == FALSE)
+		{
+			return CreateExitCode(Result_LaunchFail);
+		}
+	}
+
+	OnScopeExit([&ProcessInfo] ()
+	{
+		CloseHandle(ProcessInfo.hProcess);
+		CloseHandle(ProcessInfo.hThread);
+	});
+#endif // !TS_BUILD_DEBUG
 
 	TS_LOG("Waiting on begun event");
 	int Ret = Result_Ok;
@@ -629,9 +670,6 @@ static int MainFork(int ArgC, char** ArgV)
 	}
 
 	DaemonThread.join();
-#else
-	CloseHandle(ProcessInfo.hProcess);
-	CloseHandle(ProcessInfo.hThread);
 #endif
 
 	TS_LOG("Complete (ret=%d)", Ret);
diff --git a/Engine/Source/Programs/UnrealTraceServer/src/Pch.h b/Engine/Source/Programs/UnrealTraceServer/src/Pch.h
--- a/Engine/Source/Programs/UnrealTraceServer/src/Pch.h
+++ b/Engine/Source/Programs/UnrealTraceServer/src/Pch.h
@@ -73,8 +73,10 @@
 #	define WIN32_LEAN_AND_MEAN
 #	define NOMINMAX
 #	include <Windows.h>
+#	include <winsafer.h>
 #	include <Shlobj.h>
 #	include <shellapi.h>
+#	pragma comment(lib, "Advapi32.lib")
 #	pragma comment(lib, "Shell32.lib")
 #	pragma comment(lib, "User32.lib")
 #endif
