Include GetLastError() just in case setting the working directory happens to fail
223e89d / Martin Ridgers / [REDACTED_PHONE]:30:04 +0200

diff --git a/Engine/Source/Programs/UnrealTraceServer/src/Main.cpp b/Engine/Source/Programs/UnrealTraceServer/src/Main.cpp
--- a/Engine/Source/Programs/UnrealTraceServer/src/Main.cpp
+++ b/Engine/Source/Programs/UnrealTraceServer/src/Main.cpp
@@ -690,7 +690,7 @@ static int MainDaemon(int ArgC, char** ArgV)
 		BinDir = BinDir.parent_path();
 		std::filesystem::current_path(BinDir, ErrorCode);
 		const char* Result = ErrorCode ? "Failed" : "Succeeded";
-		TS_LOG("%s setting '%ls'", Result, BinDir.c_str());
+		TS_LOG("%s setting '%ls' (gle=%d)", Result, BinDir.c_str(), GetLastError());
 	}
 	else
 	{
