Fallback to %USERPROFILE% if SHGetFolderPath() fails
899150b / Martin Ridgers / [REDACTED_PHONE]:36:21 +0200

diff --git a/Engine/Source/Programs/UnrealTraceServer/src/Main.cpp b/Engine/Source/Programs/UnrealTraceServer/src/Main.cpp
--- a/Engine/Source/Programs/UnrealTraceServer/src/Main.cpp
+++ b/Engine/Source/Programs/UnrealTraceServer/src/Main.cpp
@@ -95,7 +95,15 @@ static void GetUnrealTraceHome(std::filesystem::path& Out, bool Make=false)
 {
 #if TS_USING(TS_PLATFORM_WINDOWS)
 	wchar_t Buffer[MAX_PATH];
-	SHGetFolderPathW(nullptr, CSIDL_LOCAL_APPDATA, nullptr, SHGFP_TYPE_CURRENT, Buffer);
+	auto Ok = SHGetFolderPathW(nullptr, CSIDL_LOCAL_APPDATA, nullptr, SHGFP_TYPE_CURRENT, Buffer);
+	if (Ok != S_OK)
+	{
+		uint32 Ret = GetEnvironmentVariableW(L"USERPROFILE", Buffer, TS_ARRAY_COUNT(Buffer));
+		if (Ret == 0 || Ret >= TS_ARRAY_COUNT(Buffer))
+		{
+			return;
+		}
+	}
 	Out = Buffer;
 	Out /= "UnrealEngine/Common/UnrealTrace";
 #else
