Added missing include
9a52de8 / Martin Ridgers / [REDACTED_PHONE]:18:14 +0200

diff --git a/Engine/Source/Programs/UnrealTraceServer/src/Main.cpp b/Engine/Source/Programs/UnrealTraceServer/src/Main.cpp
--- a/Engine/Source/Programs/UnrealTraceServer/src/Main.cpp
+++ b/Engine/Source/Programs/UnrealTraceServer/src/Main.cpp
@@ -14,6 +14,7 @@
 #	include <semaphore.h>
 #	include <sched.h>
 #	include <signal.h>
+#	include <cstdarg>
 #	include <sys/mman.h>
 #	include <unistd.h>
 #endif
