Split the two schemes of starting a child UnrealTraceServer process into separate functions
(git:2778f91 - Martin Ridgers - [REDACTED_PHONE]:25:32 +0100)

diff --git a/Engine/Source/Programs/UnrealTraceServer/src/Main.cpp b/Engine/Source/Programs/UnrealTraceServer/src/Main.cpp
--- a/Engine/Source/Programs/UnrealTraceServer/src/Main.cpp
+++ b/Engine/Source/Programs/UnrealTraceServer/src/Main.cpp
@@ -409,6 +409,63 @@ private:
 				FWinHandle(const FWinHandle&&) = delete;
 };
 
+
+
+////////////////////////////////////////////////////////////////////////////////
+static bool LaunchUnfancy(const wchar_t* Binary, wchar_t* CommandLine)
+{
+	/* No frills spawning of a child process */
+
+	uint32 CreateProcFlags = CREATE_BREAKAWAY_FROM_JOB;
+	STARTUPINFOW StartupInfo = { sizeof(STARTUPINFOW) };
+	PROCESS_INFORMATION ProcessInfo = {};
+
+	BOOL bOk = CreateProcessW(Binary, CommandLine, nullptr, nullptr, FALSE,
+		CreateProcFlags, nullptr, nullptr, &StartupInfo, &ProcessInfo);
+
+	if (bOk == FALSE)
+	{
+		return false;
+	}
+
+	CloseHandle(ProcessInfo.hProcess);
+	CloseHandle(ProcessInfo.hThread);
+	return true;
+}
+
+////////////////////////////////////////////////////////////////////////////////
+static bool LaunchUnprivileged(const wchar_t* Binary, wchar_t* CommandLine)
+{
+	/* Ordinary child process without inheriting any administrator rights */
+
+	uint32 CreateProcFlags = CREATE_BREAKAWAY_FROM_JOB;
+	STARTUPINFOW StartupInfo = { sizeof(STARTUPINFOW) };
+	PROCESS_INFORMATION ProcessInfo = {};
+
+	SAFER_LEVEL_HANDLE SafeLevel = nullptr;
+	BOOL bOk = SaferCreateLevel(SAFER_SCOPEID_USER,
+		SAFER_LEVELID_NORMALUSER, SAFER_LEVEL_OPEN, &SafeLevel, nullptr);
+	if (bOk == FALSE)
+	{
+		return false;
+	}
+
+	HANDLE AccessToken;
+	if (SaferComputeTokenFromLevel(SafeLevel, nullptr, &AccessToken, 0, nullptr))
+	{
+		bOk = CreateProcessAsUserW(AccessToken, Binary, CommandLine,
+			nullptr, nullptr, FALSE, CreateProcFlags, nullptr, nullptr,
+			&StartupInfo, &ProcessInfo);
+
+		CloseHandle(AccessToken);
+	}
+
+	SaferCloseLevel(SafeLevel);
+	return (bOk == TRUE);
+}
+
+
+
 ////////////////////////////////////////////////////////////////////////////////
 static const wchar_t*	GIpcName					= L"Local\\UnrealTraceInstance";
 static const int32		GIpcSize					= 4 << 10;
@@ -603,57 +660,16 @@ static int MainFork(int ArgC, char** ArgV)
 #if TS_USING(TS_DAEMON_THREAD)
 	std::thread DaemonThread([] () { MainDaemon(0, nullptr); });
 #else
-	uint32 CreateProcFlags = CREATE_BREAKAWAY_FROM_JOB;
-	STARTUPINFOW StartupInfo = { sizeof(STARTUPINFOW) };
-	PROCESS_INFORMATION ProcessInfo = {};
-
-	// Limit the authority of the daemon.
-	SAFER_LEVEL_HANDLE SafeLevel = nullptr;
-	BOOL bCanLaunchSafely = SaferCreateLevel(SAFER_SCOPEID_USER,
-		SAFER_LEVELID_NORMALUSER, SAFER_LEVEL_OPEN, &SafeLevel, nullptr);
-	if (bCanLaunchSafely == TRUE)
+	wchar_t CommandLine[] = { L"UnrealTraceServer.exe daemon" };
+	if (!LaunchUnprivileged(DestPath.c_str(), CommandLine))
 	{
-		HANDLE AccessToken;
-		if (SaferComputeTokenFromLevel(SafeLevel, nullptr, &AccessToken, 0, nullptr))
-		{
-			BOOL bOk = CreateProcessAsUserW(
-				AccessToken,
-				DestPath.c_str(),
-				LPWSTR(L"UnrealTraceServer.exe daemon"),
-				nullptr, nullptr, FALSE, CreateProcFlags, nullptr, nullptr,
-				&StartupInfo, &ProcessInfo);
-
-			if (bOk == FALSE)
-			{
-				bCanLaunchSafely = bOk;
-			}
-
-			CloseHandle(AccessToken);
-		}
-
-		SaferCloseLevel(SafeLevel);
-	}
-
-	// Fallback to a normal CreateProc() call if using a limited token failed
-	if (bCanLaunchSafely == FALSE)
-	{
-		BOOL bOk = CreateProcessW(
-			DestPath.c_str(),
-			LPWSTR(L"UnrealTraceServer.exe daemon"),
-			nullptr, nullptr, FALSE, CreateProcFlags, nullptr, nullptr,
-			&StartupInfo, &ProcessInfo);
-
-		if (bOk == FALSE)
+		TS_LOG("Unprivileged launch failed (gle=%d)", GetLastError());
+		if (!LaunchUnfancy(DestPath.c_str(), CommandLine))
 		{
+			TS_LOG("Launch failed (gle=%d)", GetLastError());
 			return CreateExitCode(Result_LaunchFail);
 		}
 	}
-
-	OnScopeExit([&ProcessInfo] ()
-	{
-		CloseHandle(ProcessInfo.hProcess);
-		CloseHandle(ProcessInfo.hThread);
-	});
 #endif // !TS_BUILD_DEBUG
 
 	TS_LOG("Waiting on begun event");
