Added a way to spawn an unelvated child process from an elevated parent
(git:d7ac559 - Martin Ridgers - [REDACTED_PHONE]:26:34 +0100)

diff --git a/Engine/Source/Programs/UnrealTraceServer/src/Main.cpp b/Engine/Source/Programs/UnrealTraceServer/src/Main.cpp
--- a/Engine/Source/Programs/UnrealTraceServer/src/Main.cpp
+++ b/Engine/Source/Programs/UnrealTraceServer/src/Main.cpp
@@ -464,6 +464,84 @@ static bool LaunchUnprivileged(const wchar_t* Binary, wchar_t* CommandLine)
 	return (bOk == TRUE);
 }
 
+////////////////////////////////////////////////////////////////////////////////
+static bool LaunchUnelevated(const wchar_t* Binary, wchar_t* CommandLine)
+{
+	/* Launches a binary with the shell as its parent. The shell (such as
+	   Explorer) should be an unelevated process. */
+
+	// No sense in using this route if we are not elevated in the first place
+	if (IsUserAnAdmin() == FALSE)
+	{
+		return LaunchUnprivileged(Binary, CommandLine);
+	}
+
+	// Get the users' shell process and open it for process creation
+	HWND ShellWnd = GetShellWindow();
+	if (ShellWnd == nullptr)
+	{
+		return false;
+	}
+
+	DWORD ShellPid;
+	GetWindowThreadProcessId(ShellWnd, &ShellPid);
+
+	FWinHandle Process = OpenProcess(PROCESS_CREATE_PROCESS, FALSE, ShellPid);
+	if (!Process.IsValid())
+	{
+		return false;
+	}
+
+	// Creating a process as a child of another process is done by setting a
+	// thread-attribute list on the startup info passed to CreateProcess()
+	SIZE_T AttrListSize;
+	InitializeProcThreadAttributeList(nullptr, 1, 0, &AttrListSize);
+
+	auto AttrList = (PPROC_THREAD_ATTRIBUTE_LIST)malloc(AttrListSize);
+	OnScopeExit([&] () { free(AttrList); });
+
+	if (!InitializeProcThreadAttributeList(AttrList, 1, 0, &AttrListSize))
+	{
+		return false;
+	}
+
+	BOOL bOk = UpdateProcThreadAttribute(AttrList, 0, PROC_THREAD_ATTRIBUTE_PARENT_PROCESS,
+		(HANDLE*)&Process, sizeof(Process), nullptr, nullptr);
+	if (!bOk)
+	{
+		return false;
+	}
+
+	// By this point we know we are an elevated process. It is not allowed to
+	// create a process as a child of another unelevated process that share our
+	// elevated console window if we have one. So we'll need to create a new one.
+	uint32 CreateProcFlags = CREATE_BREAKAWAY_FROM_JOB|EXTENDED_STARTUPINFO_PRESENT;
+	if (GetConsoleWindow() != nullptr)
+	{
+		CreateProcFlags |= CREATE_NEW_CONSOLE;
+	}
+	else
+	{
+		CreateProcFlags |= DETACHED_PROCESS;
+	}
+
+	// Everything is set up now so we can proceed and launch the process
+	STARTUPINFOEXW StartupInfo = { sizeof(STARTUPINFOEXW) };
+	StartupInfo.lpAttributeList = AttrList;
+	PROCESS_INFORMATION ProcessInfo = {};
+
+	bOk = CreateProcessW(Binary, CommandLine, nullptr, nullptr, FALSE,
+		CreateProcFlags, nullptr, nullptr, &StartupInfo.StartupInfo, &ProcessInfo);
+	if (bOk == FALSE)
+	{
+		return false;
+	}
+
+	CloseHandle(ProcessInfo.hProcess);
+	CloseHandle(ProcessInfo.hThread);
+	return true;
+}
+
 
 
 ////////////////////////////////////////////////////////////////////////////////
@@ -661,13 +739,17 @@ static int MainFork(int ArgC, char** ArgV)
 	std::thread DaemonThread([] () { MainDaemon(0, nullptr); });
 #else
 	wchar_t CommandLine[] = { L"UnrealTraceServer.exe daemon" };
-	if (!LaunchUnprivileged(DestPath.c_str(), CommandLine))
+	if (!LaunchUnelevated(DestPath.c_str(), CommandLine))
 	{
-		TS_LOG("Unprivileged launch failed (gle=%d)", GetLastError());
-		if (!LaunchUnfancy(DestPath.c_str(), CommandLine))
+		TS_LOG("Unelevated launch failed (gle=%d)", GetLastError());
+		if (!LaunchUnprivileged(DestPath.c_str(), CommandLine))
 		{
-			TS_LOG("Launch failed (gle=%d)", GetLastError());
-			return CreateExitCode(Result_LaunchFail);
+			TS_LOG("Unprivileged launch failed (gle=%d)", GetLastError());
+			if (!LaunchUnfancy(DestPath.c_str(), CommandLine))
+			{
+				TS_LOG("Launch failed (gle=%d)", GetLastError());
+				return CreateExitCode(Result_LaunchFail);
+			}
 		}
 	}
 #endif // !TS_BUILD_DEBUG
