Version bump
(git:8795212 - Martin Ridgers - [REDACTED_PHONE]:46:10 +0100)

diff --git a/Engine/Source/Programs/UnrealTraceServer/src/Version.h b/Engine/Source/Programs/UnrealTraceServer/src/Version.h
--- a/Engine/Source/Programs/UnrealTraceServer/src/Version.h
+++ b/Engine/Source/Programs/UnrealTraceServer/src/Version.h
@@ -7,7 +7,7 @@
 #endif
 
 #if !defined(TS_VERSION_MINOR)
-#	define TS_VERSION_MINOR		8
+#	define TS_VERSION_MINOR		9
 #endif
 
 #define TS_ICON_ID				493
