[UnrealTraceServer] Use std::filesystem::path throughout the process.

In order to support UTF-8 compatible filenames, keep the path representation of trace files as long as possible. When transmitting them as CBOR convert them using `u8string` function.

Also use non throwing overloads of some `std::filesystem` functions to avoid crashing the server if a file is deleted before a file operation is performed.

diff --git a/Engine/Source/Programs/UnrealTraceServer/src/AsioFile.cpp b/Engine/Source/Programs/UnrealTraceServer/src/AsioFile.cpp
--- a/Engine/Source/Programs/UnrealTraceServer/src/AsioFile.cpp
+++ b/Engine/Source/Programs/UnrealTraceServer/src/AsioFile.cpp
@@ -98,11 +98,10 @@ bool FAsioFile::ReadSome(void* Dest, uint32 DestSize, FAsioIoSink* Sink, uint32
 }
 
 ////////////////////////////////////////////////////////////////////////////////
-FAsioWriteable* FAsioFile::WriteFile(asio::io_context& IoContext, const char* Path)
+FAsioWriteable* FAsioFile::WriteFile(asio::io_context& IoContext, const FPath& Path)
 {
 #if TS_USING(TS_PLATFORM_WINDOWS)
-	FWinApiStr PathW(Path);
-	HANDLE Handle = CreateFileW(PathW, GENERIC_WRITE, FILE_SHARE_READ, nullptr,
+	HANDLE Handle = CreateFileW(Path.c_str(), GENERIC_WRITE, FILE_SHARE_READ, nullptr,
 		CREATE_ALWAYS, FILE_FLAG_OVERLAPPED|FILE_ATTRIBUTE_NORMAL, nullptr);
 	if (Handle == INVALID_HANDLE_VALUE)
 	{
@@ -110,7 +109,7 @@ FAsioWriteable* FAsioFile::WriteFile(asio::io_context& IoContext, const char* Pa
 	}
 	return new FAsioFile(IoContext, uintptr_t(Handle));
 #else
-	int File = open(Path, O_WRONLY|O_CREAT, 0666);
+	int File = open(Path.c_str(), O_WRONLY|O_CREAT, 0666);
 	if (!File)
 	{
 		return nullptr;
@@ -120,11 +119,10 @@ FAsioWriteable* FAsioFile::WriteFile(asio::io_context& IoContext, const char* Pa
 }
 
 ////////////////////////////////////////////////////////////////////////////////
-FAsioReadable* FAsioFile::ReadFile(asio::io_context& IoContext, const char* Path)
+FAsioReadable* FAsioFile::ReadFile(asio::io_context& IoContext, const FPath& Path)
 {
 #if TS_USING(TS_PLATFORM_WINDOWS)
-	FWinApiStr PathW(Path);
-	HANDLE Handle = CreateFileW(PathW, GENERIC_READ, FILE_SHARE_READ|FILE_SHARE_WRITE,
+	HANDLE Handle = CreateFileW(Path.c_str(), GENERIC_READ, FILE_SHARE_READ | FILE_SHARE_WRITE,
 		nullptr, OPEN_EXISTING, FILE_FLAG_OVERLAPPED, nullptr);
 	if (Handle == INVALID_HANDLE_VALUE)
 	{
@@ -132,7 +130,7 @@ FAsioReadable* FAsioFile::ReadFile(asio::io_context& IoContext, const char* Path
 	}
 	return new FAsioFile(IoContext, uintptr_t(Handle));
 #else
-	int File = open(Path, O_RDONLY, 0444);
+	int File = open(Path.c_str(), O_RDONLY, 0444);
 	if (!File)
 	{
 		return nullptr;
diff --git a/Engine/Source/Programs/UnrealTraceServer/src/AsioFile.h b/Engine/Source/Programs/UnrealTraceServer/src/AsioFile.h
--- a/Engine/Source/Programs/UnrealTraceServer/src/AsioFile.h
+++ b/Engine/Source/Programs/UnrealTraceServer/src/AsioFile.h
@@ -5,6 +5,7 @@
 #include "Asio.h"
 
 #include "AsioIoable.h"
+#include "Foundation.h"
 #if !TS_USING(TS_PLATFORM_WINDOWS)
 #include "asio/posix/stream_descriptor.hpp"
 #endif
@@ -16,8 +17,8 @@ class FAsioFile
 {
 public:
 										FAsioFile(asio::io_context& IoContext, uintptr_t OsHandle);
-	static FAsioWriteable*				WriteFile(asio::io_context& IoContext, const char* Path);
-	static FAsioReadable*				ReadFile(asio::io_context& IoContext, const char* Path);
+	static FAsioWriteable*				WriteFile(asio::io_context& IoContext, const FPath& Path);
+	static FAsioReadable*				ReadFile(asio::io_context& IoContext, const FPath& Path);
 	virtual bool						IsOpen() const override;
 	virtual void						Close() override;
 	virtual bool						Write(const void* Src, uint32 Size, FAsioIoSink* Sink, uint32 Id) override;
diff --git a/Engine/Source/Programs/UnrealTraceServer/src/Cbor.cpp b/Engine/Source/Programs/UnrealTraceServer/src/Cbor.cpp
--- a/Engine/Source/Programs/UnrealTraceServer/src/Cbor.cpp
+++ b/Engine/Source/Programs/UnrealTraceServer/src/Cbor.cpp
@@ -48,7 +48,7 @@ FStringView	FCborContext::AsString() const
 		return FStringView();
 	}
 
-    return FStringView((const char*)(NextCursor - Param), Param);
+    return FStringView((const char*)(NextCursor - Param), std::size_t(Param));
 }
 
 ////////////////////////////////////////////////////////////////////////////////
diff --git a/Engine/Source/Programs/UnrealTraceServer/src/Foundation.h b/Engine/Source/Programs/UnrealTraceServer/src/Foundation.h
--- a/Engine/Source/Programs/UnrealTraceServer/src/Foundation.h
+++ b/Engine/Source/Programs/UnrealTraceServer/src/Foundation.h
@@ -54,6 +54,8 @@ struct FString
 	void			operator += (const FStringView& Rhs)			{ Super::operator += ((FStringView::Super&)(Rhs)); }
 };
 
+////////////////////////////////////////////////////////////////////////////////
+typedef std::filesystem::path FPath;
 
 
 ////////////////////////////////////////////////////////////////////////////////
diff --git a/Engine/Source/Programs/UnrealTraceServer/src/Main.cpp b/Engine/Source/Programs/UnrealTraceServer/src/Main.cpp
--- a/Engine/Source/Programs/UnrealTraceServer/src/Main.cpp
+++ b/Engine/Source/Programs/UnrealTraceServer/src/Main.cpp
@@ -97,7 +97,7 @@ struct TOnScopeExit
 
 
 ////////////////////////////////////////////////////////////////////////////////
-static void GetUnrealTraceHome(std::filesystem::path& Out, bool Make=false)
+static void GetUnrealTraceHome(FPath& Out, bool Make=false)
 {
 #if TS_USING(TS_PLATFORM_WINDOWS)
 	wchar_t Buffer[MAX_PATH];
@@ -160,13 +160,13 @@ FLogging* FLogging::Instance = nullptr;
 FLogging::FLogging()
 {
 	// Find where the logs should be written to. Make sure it exists.
-	std::filesystem::path LogDir;
+	FPath LogDir;
 	GetUnrealTraceHome(LogDir, true);
 
 	// Fetch all existing logs.
 	struct FExistingLog
 	{
-		std::filesystem::path	Path;
+		FPath	Path;
 		uint32					Index;
 
 		int32 operator < (const FExistingLog& Rhs) const
@@ -206,7 +206,7 @@ FLogging::FLogging()
 		++LogIndex;
 		char LogName[128];
 		snprintf(LogName, TS_ARRAY_COUNT(LogName), "Server_%d.log", LogIndex);
-		std::filesystem::path LogPath = LogDir / LogName;
+		FPath LogPath = LogDir / LogName;
 
 #if TS_USING(TS_PLATFORM_WINDOWS)
 		File = _wfopen(LogPath.c_str(), L"wbxN");
@@ -291,7 +291,7 @@ struct FLoggingScope
 ////////////////////////////////////////////////////////////////////////////////
 struct FStoreOptions
 {
-	fs::path	Dir;
+	FPath		Dir;
 	int			Port			= 1989;
 	int			RecorderPort	= 1981;
 };
@@ -717,7 +717,7 @@ static int MainFork(int ArgC, char** ArgV)
 	TS_LOG("Binary located at '%ls'", BinPath);
 
 	// Calculate where to store the binaries.
-	std::filesystem::path DestPath;
+	FPath DestPath;
 	GetUnrealTraceHome(DestPath);
 	{
 		wchar_t Buffer[64];
@@ -747,7 +747,7 @@ static int MainFork(int ArgC, char** ArgV)
 		DWORD OurPid = GetCurrentProcessId();
 		wchar_t Buffer[16];
 		_snwprintf(Buffer, TS_ARRAY_COUNT(Buffer), L"_%08x", OurPid);
-		std::filesystem::path TempPath = DestPath;
+		FPath TempPath = DestPath;
 		TempPath += Buffer;
 		if (!std::filesystem::copy_file(BinPath, TempPath))
 		{
@@ -845,7 +845,7 @@ static int MainDaemon(int ArgC, char** ArgV)
 	if (GetLastError() != ERROR_INSUFFICIENT_BUFFER && BinPathLen > 0)
 	{
 		std::error_code ErrorCode;
-		std::filesystem::path BinDir(BinPath);
+		FPath BinDir(BinPath);
 		BinDir = BinDir.parent_path();
 		std::filesystem::current_path(BinDir, ErrorCode);
 		const char* Result = ErrorCode ? "Failed" : "Succeeded";
@@ -894,7 +894,7 @@ static int MainDaemon(int ArgC, char** ArgV)
 
 		if (Options.Dir.empty())
 		{
-			std::filesystem::path StoreDir;
+			FPath StoreDir;
 			GetUnrealTraceHome(StoreDir);
 			StoreDir /= "Store";
 			Options.Dir = StoreDir;
@@ -926,7 +926,7 @@ static int MainDaemon(int ArgC, char** ArgV)
 }
 
 ////////////////////////////////////////////////////////////////////////////////
-int WinMain(HINSTANCE, HINSTANCE, LPSTR, int)
+int WINAPI WinMain(HINSTANCE, HINSTANCE, LPSTR, int)
 {
 	// Get command line arguments and convert to UTF8
 	const wchar_t* CommandLine = GetCommandLineW();
@@ -967,7 +967,7 @@ int WinMain(HINSTANCE, HINSTANCE, LPSTR, int)
 static int MainDaemon(int, char**, pid_t);
 
 ////////////////////////////////////////////////////////////////////////////////
-static std::filesystem::path GetLockFilePath()
+static FPath GetLockFilePath()
 {
 	return "/tmp/UnrealTraceServer.pid";
 }
@@ -1016,7 +1016,7 @@ static int MainKillImpl(int ArgC, char** ArgV, pid_t DaemonPid)
 static int MainKill(int ArgC, char** ArgV)
 {
 	// Open the pid file to detect an existing instance
-	std::filesystem::path DotPidPath = GetLockFilePath();
+	FPath DotPidPath = GetLockFilePath();
 	TS_LOG("Checking for a '%s' lock file", DotPidPath.c_str());
 	int DotPidFd = open(DotPidPath.c_str(), O_RDONLY);
 	if (DotPidFd < 0)
@@ -1050,7 +1050,7 @@ static int MainKill(int ArgC, char** ArgV)
 static int MainFork(int ArgC, char** ArgV)
 {
 	// Open the pid file to detect an existing instance
-	std::filesystem::path DotPidPath = GetLockFilePath();
+	FPath DotPidPath = GetLockFilePath();
 	TS_LOG("Checking for a '%s' lock file", DotPidPath.c_str());
 	for (int DotPidFd = open(DotPidPath.c_str(), O_RDONLY); DotPidFd >= 0; )
 	{
@@ -1174,7 +1174,7 @@ static int MainFork(int ArgC, char** ArgV)
 static int MainDaemon(int ArgC, char** ArgV, pid_t ParentPid)
 {
 	// We expect that there is no lock file on disk if we've got this far.
-	std::filesystem::path DotPidPath = GetLockFilePath();
+	FPath DotPidPath = GetLockFilePath();
 	TS_LOG("Opening the lock file '%s'", DotPidPath.c_str());
 	int DotPidFd = open(DotPidPath.c_str(), O_CREAT|O_WRONLY, 0666);
 	if (DotPidFd < 0)
@@ -1202,7 +1202,7 @@ static int MainDaemon(int ArgC, char** ArgV, pid_t ParentPid)
 
 		if (Options.Dir.empty())
 		{
-			std::filesystem::path StoreDir;
+			FPath StoreDir;
 			GetUnrealTraceHome(StoreDir);
 			StoreDir /= "Store";
 			Options.Dir = StoreDir;
@@ -1312,7 +1312,7 @@ int main(int ArgC, char** ArgV)
 		puts(  "data for analysis. TCP ports 1981 and 1989 are used, where the former receives");
 		puts(  "trace data, and the latter is used by tools to query the server's store.");
 
-		std::filesystem::path HomeDir;
+		FPath HomeDir;
 		GetUnrealTraceHome(HomeDir);
 		HomeDir.make_preferred();
 		std::string HomeDirU8 = HomeDir.string();
diff --git a/Engine/Source/Programs/UnrealTraceServer/src/Store.cpp b/Engine/Source/Programs/UnrealTraceServer/src/Store.cpp
--- a/Engine/Source/Programs/UnrealTraceServer/src/Store.cpp
+++ b/Engine/Source/Programs/UnrealTraceServer/src/Store.cpp
@@ -80,7 +80,7 @@ public:
 	HandlerType Handler;
 private:
 	bool bIsRunning = false;
-	std::filesystem::path StoreDir;
+	FPath StoreDir;
 };
 
 void MacCallback(ConstFSEventStreamRef StreamRef,
@@ -144,29 +144,15 @@ void FStore::FDirWatcher::async_wait(HandlerType InHandler)
 
 
 ////////////////////////////////////////////////////////////////////////////////
-FStore::FTrace::FTrace(const char* InPath)
+FStore::FTrace::FTrace(const FPath& InPath)
 : Path(InPath)
 {
-	// Extract the trace's name
-	const char* Dot = std::strrchr(*Path, '.');
-	if (Dot == nullptr)
-	{
-		Dot = *Path;
-	}
-
-	for (const char* c = Dot; c > *Path; --c)
-	{
-		if (c[-1] == '\\' || c[-1] == '/')
-		{
-			Name = FStringView(c, int32(Dot - c));
-			break;
-		}
-	}
-
-	Id = QuickStoreHash(Name);
+	Name = InPath.filename().string();
+	Id = QuickStoreHash(*Name);
 
+	std::error_code Ec;
 	// Calculate that trace's timestamp. Bias in seconds then convert to 0.1us.
-	std::filesystem::file_time_type LastWriteTime = std::filesystem::last_write_time(InPath);
+	std::filesystem::file_time_type LastWriteTime = std::filesystem::last_write_time(Path, Ec);
 	auto LastWriteDuration = LastWriteTime.time_since_epoch();
 	Timestamp = std::chrono::duration_cast<std::chrono::seconds>(LastWriteDuration).count();
 	Timestamp += FsToUnrealEpochBiasSeconds;
@@ -174,9 +160,9 @@ FStore::FTrace::FTrace(const char* InPath)
 }
 
 ////////////////////////////////////////////////////////////////////////////////
-const FStringView& FStore::FTrace::GetName() const
+const FPath& FStore::FTrace::GetPath() const
 {
-	return Name;
+	return Path;
 }
 
 ////////////////////////////////////////////////////////////////////////////////
@@ -188,7 +174,8 @@ uint32 FStore::FTrace::GetId() const
 ////////////////////////////////////////////////////////////////////////////////
 uint64 FStore::FTrace::GetSize() const
 {
-	return std::filesystem::file_size(*Path);
+	std::error_code Ec;
+	return std::filesystem::file_size(Path, Ec);
 }
 
 ////////////////////////////////////////////////////////////////////////////////
@@ -200,7 +187,7 @@ uint64 FStore::FTrace::GetTimestamp() const
 
 
 ////////////////////////////////////////////////////////////////////////////////
-FStore::FMount::FMount(const fs::path& InDir)
+FStore::FMount::FMount(const FPath& InDir)
 : Id(QuickStoreHash(InDir.c_str()))
 {
 	std::error_code ErrorCode;
@@ -252,7 +239,7 @@ FStore::FTrace* FStore::FMount::GetTrace(uint32 Id) const
 }
 
 ////////////////////////////////////////////////////////////////////////////////
-FStore::FTrace* FStore::FMount::AddTrace(const char* Path)
+FStore::FTrace* FStore::FMount::AddTrace(const FPath& Path)
 {
 	FTrace NewTrace(Path);
 
@@ -278,13 +265,14 @@ uint32 FStore::FMount::Refresh()
 			continue;
 		}
 
-		std::filesystem::path Extension = DirItem.path().extension();
+		FPath Extension = DirItem.path().extension();
 		if (Extension != ".utrace")
 		{
 			continue;
 		}
 
-		FTrace* Trace = AddTrace(DirItem.path().string().c_str());
+		FTrace* Trace = AddTrace(DirItem.path());
+
 		if (Trace != nullptr)
 		{
 			ChangeSerial += Trace->GetId();
@@ -296,10 +284,10 @@ uint32 FStore::FMount::Refresh()
 
 
 ////////////////////////////////////////////////////////////////////////////////
-FStore::FStore(asio::io_context& InIoContext, const fs::path& InStoreDir)
+FStore::FStore(asio::io_context& InIoContext, const FPath& InStoreDir)
 : IoContext(InIoContext)
 {
-	fs::path StoreDir = InStoreDir / "001";
+	FPath StoreDir = InStoreDir / "001";
 	AddMount(StoreDir);
 
 	Refresh();
@@ -349,7 +337,7 @@ void FStore::Close()
 }
 
 ////////////////////////////////////////////////////////////////////////////////
-bool FStore::AddMount(const fs::path& Dir)
+bool FStore::AddMount(const FPath& Dir)
 {
 	FMount* Mount = new FMount(Dir);
 	Mounts.Add(Mount);
@@ -502,31 +490,6 @@ FStore::FTrace* FStore::GetTrace(uint32 Id, FMount** OutMount) const
 ////////////////////////////////////////////////////////////////////////////////
 FStore::FNewTrace FStore::CreateTrace()
 {
-	FString TracePath;
-#if 0
-	bool bOk = false;
-	for (int i = 0; i < 256; ++i)
-	{
-		uint32 TraceId = ++LastTraceId;
-
-		TracePath = StoreDir;
-		TracePath.Appendf("/%05d", TraceId);
-
-		if (!std::filesystem::is_directory(*TracePath) && std::filesystem::create_directories(*TracePath))
-		{
-			bOk = true;
-			break;
-		}
-	}
-
-	if (!bOk)
-	{
-		return {};
-	}
-
-	TracePath += "/data.utrace";
-#else
-
 	// N.B. Not thread safe!?
 	char Prefix[24];
 	std::time_t Now = std::time(nullptr);
@@ -535,27 +498,24 @@ FStore::FNewTrace FStore::CreateTrace()
 
 	FMount* DefaultMount = Mounts[0];
 
-	TracePath = DefaultMount->GetDir();
-	TracePath += "/";
-	TracePath += Prefix;
+	FPath TracePath(*DefaultMount->GetDir());
+	TracePath /= Prefix;
 	TracePath += ".utrace";
-	for (uint32 Index = 0; std::filesystem::is_regular_file(*TracePath); ++Index)
+	
+	for (uint32 Index = 0; std::filesystem::is_regular_file(TracePath); ++Index)
 	{
-		char Suffix[64];
-		std::snprintf(Suffix, TS_ARRAY_COUNT(Suffix), "/%s_%02d.utrace", Prefix, Index);
-
-		TracePath = DefaultMount->GetDir();
-		TracePath += *Suffix;
+		char FilenameIndexed[27];
+		std::sprintf(FilenameIndexed, "%s_%02d.utrace", Prefix, Index);
+		TracePath.replace_filename(FPath(FilenameIndexed));
 	}
-#endif // 0
 
-	FAsioWriteable* File = FAsioFile::WriteFile(IoContext, *TracePath);
+	FAsioWriteable* File = FAsioFile::WriteFile(IoContext, TracePath);
 	if (File == nullptr)
 	{
 		return {};
 	}
 
-	FTrace* Trace = DefaultMount->AddTrace(*TracePath);
+	FTrace* Trace = DefaultMount->AddTrace(TracePath);
 	if (Trace == nullptr)
 	{
 		delete File;
@@ -581,17 +541,7 @@ FAsioReadable* FStore::OpenTrace(uint32 Id)
 		return nullptr;
 	}
 
-	FString TracePath;
-	TracePath = Mount->GetDir();
-#if 0
-	TracePath.Appendf("/%05d/data.utrace", Id);
-#else
-	TracePath += "/";
-	TracePath += Trace->GetName();
-	TracePath += ".utrace";
-#endif // 0
-
-	return FAsioFile::ReadFile(IoContext, *TracePath);
+	return FAsioFile::ReadFile(IoContext, Trace->GetPath());
 }
 
 ////////////////////////////////////////////////////////////////////////////////
diff --git a/Engine/Source/Programs/UnrealTraceServer/src/Store.h b/Engine/Source/Programs/UnrealTraceServer/src/Store.h
--- a/Engine/Source/Programs/UnrealTraceServer/src/Store.h
+++ b/Engine/Source/Programs/UnrealTraceServer/src/Store.h
@@ -23,16 +23,16 @@ public:
 	class FTrace
 	{
 	public:
-							FTrace(const char* InPath);
-		const FStringView&	GetName() const;
+							FTrace(const FPath& InPath);
+		const FPath&		GetPath() const;
 		uint32				GetId() const;
 		uint64				GetSize() const;
 		uint64				GetTimestamp() const;
 
 	private:
 		friend				FStore;
-		FString				Path;
-		FStringView			Name;
+		FPath				Path;
+		FString				Name;
 		uint64				Timestamp;
 		uint32				Id = 0;
 	};
@@ -46,27 +46,26 @@ public:
 	class FMount
 	{
 	public:
-						FMount(const fs::path& InDir);
+						FMount(const FPath& InDir);
 		uint32			GetId() const;
 		FString			GetDir() const;
 		uint32			GetTraceCount() const;
 		const FTrace*	GetTraceInfo(uint32 Index) const;
-		bool			HasTrace(uint32 Id) const;
 
 	private:
 		friend			FStore;
 		FTrace*			GetTrace(uint32 Id) const;
-		FTrace*			AddTrace(const char* Path);
+		FTrace*			AddTrace(const FPath& Path);
 		uint32			Refresh();
-		fs::path		Dir;
+		FPath			Dir;
 		TArray<FTrace*>	Traces;
 		uint32			Id;
 	};
 
-						FStore(asio::io_context& IoContext, const fs::path& InStoreDir);
+						FStore(asio::io_context& IoContext, const FPath& InStoreDir);
 						~FStore();
 	void				Close();
-	bool				AddMount(const fs::path& Dir);
+	bool				AddMount(const FPath& Dir);
 	bool				RemoveMount(uint32 Id);
 	const FMount*		GetMount(uint32 Id) const;
 	uint32				GetMountCount() const;
diff --git a/Engine/Source/Programs/UnrealTraceServer/src/StoreCborServer.cpp b/Engine/Source/Programs/UnrealTraceServer/src/StoreCborServer.cpp
--- a/Engine/Source/Programs/UnrealTraceServer/src/StoreCborServer.cpp
+++ b/Engine/Source/Programs/UnrealTraceServer/src/StoreCborServer.cpp
@@ -177,19 +177,13 @@ void FStoreCborPeer::OnTraceInfo()
 		return SendError(EStatusCode::BadRequest);
 	}
 
-	char OutName[256];
-	const FStringView& Name = Trace->GetName();
-	uint32 NameLength = std::min(uint32(sizeof(OutName)), uint32(Name.Len()));
-	for (uint32 i = 0; i < NameLength; ++i)
-	{
-		OutName[i] = char(Name[i]);
-	}
+	auto Utf8String = Trace->GetPath().filename().u8string();
 
 	TPayloadBuilder<> Builder(EStatusCode::Success);
 	Builder.AddInteger("id", Trace->GetId());
 	Builder.AddInteger("size", Trace->GetSize());
 	Builder.AddInteger("timestamp", Trace->GetTimestamp());
-	Builder.AddString("name", OutName, NameLength);
+	Builder.AddString("name", (const char*)Utf8String.c_str(), (int32)Utf8String.size());
 	SendResponse(Builder.Done());
 }
 
diff --git a/Engine/Source/Programs/UnrealTraceServer/src/StoreService.h b/Engine/Source/Programs/UnrealTraceServer/src/StoreService.h
--- a/Engine/Source/Programs/UnrealTraceServer/src/StoreService.h
+++ b/Engine/Source/Programs/UnrealTraceServer/src/StoreService.h
@@ -10,7 +10,7 @@ class FStoreService
 public:
 	struct FDesc
 	{
-		fs::path			StoreDir;
+		FPath				StoreDir;
 		int32				StorePort		= 0; // <=0:auto-assign
 		int32				RecorderPort	= 0; // 0:auto-assign, -1:off
 		int32				ThreadCount		= 0; // <=0:logical CPU count
