Add UnrealTraceServer solution and project file.

diff --git a/Engine/Source/Programs/UnrealTraceServer/Readme.md b/Engine/Source/Programs/UnrealTraceServer/Readme.md
new file mode 100644
--- /dev/null
+++ b/Engine/Source/Programs/UnrealTraceServer/Readme.md
@@ -0,0 +1,26 @@
+# Unreal Trace Server
+
+Standalone hub for recording traces from applications and serving them to analyzers (e.g. Unreal Insights).
+
+# Building
+
+## Windows
+
+1. Run "Developer Command Prompt for Visual Studio XXXX" from the start menu.
+2. Start a x64 version of the command line by executing `VsDevCmd.bat -host_arch=amd64 -arch=amd64`
+3. Navigate to this folder
+4. Execute command 
+	* `nmake` for building release configuration.
+	* `nmake /E DEBUG=1` for debug configuration.
+
+Visual Studio solution is provided for convenience. Note that changes that affect `Pch.h` requires a clean since there is no dependency checking for the precompiled header.
+
+# Making a release
+
+## Windows
+1. Run "Developer Command Prompt for Visual Studio XXXX" from the start menu.
+2. Start a x64 version of the command line by executing `VsDevCmd.bat -host_arch=amd64 -arch=amd64`
+3. Navigate to this folder
+4. Execute command 
+	* `nmake release`
+5. In Perforce check in new version of `UnrealTraceServer.exe`.
\ No newline at end of file
