Fixed possible stack overflow
abb645f / Martin Ridgers / [REDACTED_PHONE]:06:47 +0200

#rnx

diff --git a/Engine/Source/Programs/UnrealTraceServer/src/Store.cpp b/Engine/Source/Programs/UnrealTraceServer/src/Store.cpp
--- a/Engine/Source/Programs/UnrealTraceServer/src/Store.cpp
+++ b/Engine/Source/Programs/UnrealTraceServer/src/Store.cpp
@@ -504,7 +504,7 @@ FStore::FNewTrace FStore::CreateTrace()
 	
 	for (uint32 Index = 0; std::filesystem::is_regular_file(TracePath); ++Index)
 	{
-		char FilenameIndexed[27];
+		char FilenameIndexed[64];
 		std::sprintf(FilenameIndexed, "%s_%02d.utrace", Prefix, Index);
 		TracePath.replace_filename(FPath(FilenameIndexed));
 	}
