Deleted traces were not being removed from the store
b257466 / Martin Ridgers / [REDACTED_PHONE]:35:28 +0200

#rnx

diff --git a/Engine/Source/Programs/UnrealTraceServer/src/Store.cpp b/Engine/Source/Programs/UnrealTraceServer/src/Store.cpp
--- a/Engine/Source/Programs/UnrealTraceServer/src/Store.cpp
+++ b/Engine/Source/Programs/UnrealTraceServer/src/Store.cpp
@@ -257,6 +257,8 @@ FStore::FTrace* FStore::FMount::AddTrace(const FPath& Path)
 ////////////////////////////////////////////////////////////////////////////////
 uint32 FStore::FMount::Refresh()
 {
+	Traces.Empty();
+
 	uint32 ChangeSerial = 0;
 	for (auto& DirItem : std::filesystem::directory_iterator(Dir))
 	{
