Version bump
a863182 / Martin Ridgers / [REDACTED_PHONE]:55:47 +0200

#rnx

diff --git a/Engine/Source/Programs/UnrealTraceServer/src/Version.h b/Engine/Source/Programs/UnrealTraceServer/src/Version.h
--- a/Engine/Source/Programs/UnrealTraceServer/src/Version.h
+++ b/Engine/Source/Programs/UnrealTraceServer/src/Version.h
@@ -7,7 +7,7 @@
 #endif
 
 #if !defined(TS_VERSION_MINOR)
-#	define TS_VERSION_MINOR		10
+#	define TS_VERSION_MINOR		11
 #endif
 
 #define TS_ICON_ID				493
