Updated Readme.md

#rnx

diff --git a/Engine/Source/Programs/UnrealTraceServer/Readme.md b/Engine/Source/Programs/UnrealTraceServer/Readme.md
--- a/Engine/Source/Programs/UnrealTraceServer/Readme.md
+++ b/Engine/Source/Programs/UnrealTraceServer/Readme.md
@@ -17,10 +17,12 @@ Visual Studio solution is provided for convenience. Note that changes that affec
 
 # Making a release
 
+Bump `TS_VERSION_MINOR` so the auto-update mechanisms activate when users receive the newer version.
+
 ## Windows
 1. Run "Developer Command Prompt for Visual Studio XXXX" from the start menu.
 2. Start a x64 version of the command line by executing `VsDevCmd.bat -host_arch=amd64 -arch=amd64`
 3. Navigate to this folder
 4. Execute command 
-	* `nmake release`
-5. In Perforce check in new version of `UnrealTraceServer.exe`.
\ No newline at end of file
+	* `nmake clean release`
+5. In Perforce check in new version of `UnrealTraceServer.exe` and the corresponding .PDB file.
\ No newline at end of file
