Reinstated FTrace::GetName() encapsulation
7ae9c1e / Martin Ridgers / [REDACTED_PHONE]:55:04 +0200

#rnx

diff --git a/Engine/Source/Programs/UnrealTraceServer/src/Store.cpp b/Engine/Source/Programs/UnrealTraceServer/src/Store.cpp
--- a/Engine/Source/Programs/UnrealTraceServer/src/Store.cpp
+++ b/Engine/Source/Programs/UnrealTraceServer/src/Store.cpp
@@ -147,7 +147,7 @@ void FStore::FDirWatcher::async_wait(HandlerType InHandler)
 FStore::FTrace::FTrace(const FPath& InPath)
 : Path(InPath)
 {
-	Name = InPath.filename().string();
+	const FString Name = GetName();
 	Id = QuickStoreHash(*Name);
 
 	std::error_code Ec;
@@ -159,6 +159,12 @@ FStore::FTrace::FTrace(const FPath& InPath)
 	Timestamp *= 10'000'000;
 }
 
+////////////////////////////////////////////////////////////////////////////////
+FString FStore::FTrace::GetName() const
+{
+	return Path.stem().string();
+}
+
 ////////////////////////////////////////////////////////////////////////////////
 const FPath& FStore::FTrace::GetPath() const
 {
diff --git a/Engine/Source/Programs/UnrealTraceServer/src/Store.h b/Engine/Source/Programs/UnrealTraceServer/src/Store.h
--- a/Engine/Source/Programs/UnrealTraceServer/src/Store.h
+++ b/Engine/Source/Programs/UnrealTraceServer/src/Store.h
@@ -25,6 +25,7 @@ public:
 	public:
 							FTrace(const FPath& InPath);
 		const FPath&		GetPath() const;
+		FString				GetName() const;
 		uint32				GetId() const;
 		uint64				GetSize() const;
 		uint64				GetTimestamp() const;
@@ -32,7 +33,6 @@ public:
 	private:
 		friend				FStore;
 		FPath				Path;
-		FString				Name;
 		uint64				Timestamp;
 		uint32				Id = 0;
 	};
diff --git a/Engine/Source/Programs/UnrealTraceServer/src/StoreCborServer.cpp b/Engine/Source/Programs/UnrealTraceServer/src/StoreCborServer.cpp
--- a/Engine/Source/Programs/UnrealTraceServer/src/StoreCborServer.cpp
+++ b/Engine/Source/Programs/UnrealTraceServer/src/StoreCborServer.cpp
@@ -177,13 +177,13 @@ void FStoreCborPeer::OnTraceInfo()
 		return SendError(EStatusCode::BadRequest);
 	}
 
-	auto Utf8String = Trace->GetPath().filename().u8string();
+	auto Name = Trace->GetName();
 
 	TPayloadBuilder<> Builder(EStatusCode::Success);
 	Builder.AddInteger("id", Trace->GetId());
 	Builder.AddInteger("size", Trace->GetSize());
 	Builder.AddInteger("timestamp", Trace->GetTimestamp());
-	Builder.AddString("name", (const char*)Utf8String.c_str(), (int32)Utf8String.size());
+	Builder.AddString("name", *Name);
 	SendResponse(Builder.Done());
 }
 
