Updated icon for UnrealTraceServer.exe

diff --git a/Engine/Source/Programs/UnrealTraceServer/src/Icon/make_ico.py b/Engine/Source/Programs/UnrealTraceServer/src/Icon/make_ico.py
--- a/Engine/Source/Programs/UnrealTraceServer/src/Icon/make_ico.py
+++ b/Engine/Source/Programs/UnrealTraceServer/src/Icon/make_ico.py
@@ -1,11 +1,15 @@
+# Copyright Epic Games, Inc. All Rights Reserved.
+
 from PIL import Image
 
-images = []
-sizes = [255, 128, 64, 32, 48, 16]
-for size in sizes:
-    print(f"Loading ue_{size}.png")
-    image = Image.open(f"ue_{size}.png")
-    images.append(image)
+def MakeIcon(prefix, output):
+    sizes = [256, 64, 48, 32, 24, 16]
+    images = []
+    for size in sizes:
+        print(f"Loading {prefix}_{size}.png")
+        image = Image.open(f"{prefix}_{size}.png")
+        images.append(image)
+    print(f"Creating {output}")
+    images[0].save(f"{output}", append_images=images)
 
-print("Creating ../Unreal.ico")
-images[0].save("../Unreal.ico", append_images=images)
+MakeIcon("icon", "icon.ico")
diff --git a/Engine/Source/Programs/UnrealTraceServer/src/Version.rc b/Engine/Source/Programs/UnrealTraceServer/src/Version.rc
--- a/Engine/Source/Programs/UnrealTraceServer/src/Version.rc
+++ b/Engine/Source/Programs/UnrealTraceServer/src/Version.rc
@@ -19,6 +19,6 @@
 	}
 }
 
-TS_ICON_ID ICON "Unreal.ico"
+TS_ICON_ID ICON "icon.ico"
 
 /* vim: set noexpandtab : */
