Fixed support for unicode file names (for trace files) in UnrealTraceServer. Increased minor version to 12 and updated binaries for UnrealTraceServer.exe/.pdb.

diff --git a/Engine/Source/Programs/UnrealTraceServer/src/Store.cpp b/Engine/Source/Programs/UnrealTraceServer/src/Store.cpp
--- a/Engine/Source/Programs/UnrealTraceServer/src/Store.cpp
+++ b/Engine/Source/Programs/UnrealTraceServer/src/Store.cpp
@@ -162,7 +162,7 @@ FStore::FTrace::FTrace(const FPath& InPath)
 ////////////////////////////////////////////////////////////////////////////////
 FString FStore::FTrace::GetName() const
 {
-	return Path.stem().string();
+	return FString((const char*)Path.stem().u8string().c_str());
 }
 
 ////////////////////////////////////////////////////////////////////////////////
diff --git a/Engine/Source/Programs/UnrealTraceServer/src/Version.h b/Engine/Source/Programs/UnrealTraceServer/src/Version.h
--- a/Engine/Source/Programs/UnrealTraceServer/src/Version.h
+++ b/Engine/Source/Programs/UnrealTraceServer/src/Version.h
@@ -7,7 +7,7 @@
 #endif
 
 #if !defined(TS_VERSION_MINOR)
-#	define TS_VERSION_MINOR		11
+#	define TS_VERSION_MINOR		12
 #endif
 
 #define TS_ICON_ID				493
