Split op IDs across multiple lines
#rnx
(git:2ca2101 - Martin Ridgers - [REDACTED_PHONE]:29:54 +0100)

diff --git a/Engine/Source/Programs/UnrealTraceServer/src/Recorder.cpp b/Engine/Source/Programs/UnrealTraceServer/src/Recorder.cpp
--- a/Engine/Source/Programs/UnrealTraceServer/src/Recorder.cpp
+++ b/Engine/Source/Programs/UnrealTraceServer/src/Recorder.cpp
@@ -34,12 +34,19 @@ private:
 	virtual void		OnIoComplete(uint32 Id, int32 Size) override;
 	bool				ReadMetadata(int32 Size);
 	static const uint32	BufferSize = 64 * 1024;
-	enum				{ OpStart, OpSocketReadMetadata, OpSocketRead, OpFileWrite };
 	FAsioSocket			Input;
 	FAsioWriteable*		Output;
 	uint32				ActiveReadOp = OpSocketReadMetadata;
 	uint16				ControlPort = 0;
 	uint8				Buffer[BufferSize];
+
+	enum
+	{
+		OpStart,
+		OpSocketReadMetadata,
+		OpSocketRead,
+		OpFileWrite,
+	};
 };
 
 ////////////////////////////////////////////////////////////////////////////////
