Have the recorder relay create the trace itself
#rnx
(git:e317541 - Martin Ridgers - [REDACTED_PHONE]:39:12 +0100)

diff --git a/Engine/Source/Programs/UnrealTraceServer/src/Recorder.cpp b/Engine/Source/Programs/UnrealTraceServer/src/Recorder.cpp
--- a/Engine/Source/Programs/UnrealTraceServer/src/Recorder.cpp
+++ b/Engine/Source/Programs/UnrealTraceServer/src/Recorder.cpp
@@ -23,20 +23,24 @@ class FRecorderRelay
 	: public FAsioIoSink
 {
 public:
-						FRecorderRelay(asio::ip::tcp::socket& Socket, FAsioWriteable* InOutput);
+						FRecorderRelay(asio::ip::tcp::socket& Socket, FStore& InStore);
 	virtual				~FRecorderRelay();
 	bool				IsOpen();
 	void				Close();
+	uint32				GetTraceId() const;
 	uint32				GetIpAddress() const;
 	uint32				GetControlPort() const;
 
 private:
 	virtual void		OnIoComplete(uint32 Id, int32 Size) override;
+	bool				CreateTrace();
 	bool				ReadMetadata(int32 Size);
 	static const uint32	BufferSize = 64 * 1024;
 	FAsioSocket			Input;
-	FAsioWriteable*		Output;
+	FAsioWriteable*		Output = nullptr;
+	FStore&				Store;
 	uint32				ActiveReadOp = OpSocketReadMetadata;
+	uint32				TraceId = 0;
 	uint16				ControlPort = 0;
 	uint8				Buffer[BufferSize];
 
@@ -50,9 +54,9 @@ private:
 };
 
 ////////////////////////////////////////////////////////////////////////////////
-FRecorderRelay::FRecorderRelay(asio::ip::tcp::socket& Socket, FAsioWriteable* InOutput)
+FRecorderRelay::FRecorderRelay(asio::ip::tcp::socket& Socket, FStore& InStore)
 : Input(Socket)
-, Output(InOutput)
+, Store(InStore)
 {
 #if TS_USING(TS_PLATFORM_WINDOWS)
 	// Trace data is a stream and communication is one way. It is implemented
@@ -79,15 +83,21 @@ FRecorderRelay::FRecorderRelay(asio::ip::tcp::socket& Socket, FAsioWriteable* In
 	);
 #endif
 
-	OnIoComplete(OpStart, 0);
+	if (CreateTrace())
+	{
+		OnIoComplete(OpStart, 0);
+	}
 }
 
 ////////////////////////////////////////////////////////////////////////////////
 FRecorderRelay::~FRecorderRelay()
 {
 	check(!Input.IsOpen());
-	check(!Output->IsOpen());
-	delete Output;
+	if (Output != nullptr)
+	{
+		check(!Output->IsOpen());
+		delete Output;
+	}
 }
 
 ////////////////////////////////////////////////////////////////////////////////
@@ -100,7 +110,16 @@ bool FRecorderRelay::IsOpen()
 void FRecorderRelay::Close()
 {
 	Input.Close();
-	Output->Close();
+	if (Output != nullptr)
+	{
+		Output->Close();
+	}
+}
+
+////////////////////////////////////////////////////////////////////////////////
+uint32 FRecorderRelay::GetTraceId() const
+{
+	return TraceId;
 }
 
 ////////////////////////////////////////////////////////////////////////////////
@@ -115,6 +134,15 @@ uint32 FRecorderRelay::GetControlPort() const
 	return ControlPort;
 }
 
+////////////////////////////////////////////////////////////////////////////////
+bool FRecorderRelay::CreateTrace()
+{
+	FStore::FNewTrace Trace = Store.CreateTrace();
+	TraceId = Trace.Id;
+	Output = Trace.Writeable;
+	return (Output != nullptr);
+}
+
 ////////////////////////////////////////////////////////////////////////////////
 bool FRecorderRelay::ReadMetadata(int32 Size)
 {
@@ -226,7 +254,7 @@ uint32 FRecorder::FSession::GetId() const
 ////////////////////////////////////////////////////////////////////////////////
 uint32 FRecorder::FSession::GetTraceId() const
 {
-	return TraceId;
+	return Relay->GetTraceId();
 }
 
 ////////////////////////////////////////////////////////////////////////////////
@@ -296,13 +324,7 @@ const FRecorder::FSession* FRecorder::GetSessionInfo(uint32 Index) const
 ////////////////////////////////////////////////////////////////////////////////
 bool FRecorder::OnAccept(asio::ip::tcp::socket& Socket)
 {
-	FStore::FNewTrace Trace = Store.CreateTrace();
-	if (Trace.Writeable == nullptr)
-	{
-		return true;
-	}
-
-	auto* Relay = new FRecorderRelay(Socket, Trace.Writeable);
+	auto* Relay = new FRecorderRelay(Socket, Store);
 
 	uint32 IdPieces[] = {
 		Relay->GetIpAddress(),
@@ -314,7 +336,6 @@ bool FRecorder::OnAccept(asio::ip::tcp::socket& Socket)
 	FSession Session;
 	Session.Relay = Relay;
 	Session.Id = QuickStoreHash(IdPieces);
-	Session.TraceId = Trace.Id;
 	Sessions.Add(Session);
 
 	return true;
diff --git a/Engine/Source/Programs/UnrealTraceServer/src/Recorder.h b/Engine/Source/Programs/UnrealTraceServer/src/Recorder.h
--- a/Engine/Source/Programs/UnrealTraceServer/src/Recorder.h
+++ b/Engine/Source/Programs/UnrealTraceServer/src/Recorder.h
@@ -28,7 +28,6 @@ public:
 		friend				FRecorder;
 		FRecorderRelay*		Relay;
 		uint32				Id;
-		uint32				TraceId;
 	};
 
 							FRecorder(asio::io_context& IoContext, FStore& InStore);
