Split parsing of .utrace magic into its own method
#rnx
(git:31963c2 - Martin Ridgers - [REDACTED_PHONE]:03:03 +0100)

diff --git a/Engine/Source/Programs/UnrealTraceServer/src/Recorder.cpp b/Engine/Source/Programs/UnrealTraceServer/src/Recorder.cpp
--- a/Engine/Source/Programs/UnrealTraceServer/src/Recorder.cpp
+++ b/Engine/Source/Programs/UnrealTraceServer/src/Recorder.cpp
@@ -34,6 +34,7 @@ public:
 private:
 	virtual void		OnIoComplete(uint32 Id, int32 Size) override;
 	bool				CreateTrace();
+	bool				ReadMagic(uint32 Magic);
 	bool				ReadMetadata(int32 Size);
 	static const uint32	BufferSize = 64 * 1024;
 	FAsioSocket			Input;
@@ -143,6 +144,29 @@ bool FRecorderRelay::CreateTrace()
 	return (Output != nullptr);
 }
 
+////////////////////////////////////////////////////////////////////////////////
+bool FRecorderRelay::ReadMagic(uint32 Magic)
+{
+	switch (Magic)
+	{
+	/* trace with metadata data */
+	case FourCc('T', 'R', 'C', '2'):
+		break;
+
+	/* valid, but to old or wrong endian for us */
+	case FourCc('T', 'R', 'C', 'E'):
+	case FourCc('E', 'C', 'R', 'T'):
+	case FourCc('2', 'C', 'R', 'T'):
+		return true;
+
+	/* unexpected magic */
+	default:
+		return false;
+	}
+
+	return false;
+}
+
 ////////////////////////////////////////////////////////////////////////////////
 bool FRecorderRelay::ReadMetadata(int32 Size)
 {
@@ -163,20 +187,8 @@ bool FRecorderRelay::ReadMetadata(int32 Size)
 	}
 
 	Magic = *(const uint32*)(Read(sizeof(Magic)));
-	switch (Magic)
+	if (!ReadMagic(Magic))
 	{
-	/* trace with metadata data */
-	case FourCc('T', 'R', 'C', '2'):
-		break;
-
-	/* valid, but to old or wrong endian for us */
-	case FourCc('T', 'R', 'C', 'E'):
-	case FourCc('E', 'C', 'R', 'T'):
-	case FourCc('2', 'C', 'R', 'T'):
-		return true;
-
-	/* unexpected magic */
-	default:
 		return false;
 	}
 
