Removed OpStart for clarity
#rnx
(git:06eb61a - Martin Ridgers - [REDACTED_PHONE]:08:02 +0100)

diff --git a/Engine/Source/Programs/UnrealTraceServer/src/Recorder.cpp b/Engine/Source/Programs/UnrealTraceServer/src/Recorder.cpp
--- a/Engine/Source/Programs/UnrealTraceServer/src/Recorder.cpp
+++ b/Engine/Source/Programs/UnrealTraceServer/src/Recorder.cpp
@@ -47,7 +47,6 @@ private:
 
 	enum
 	{
-		OpStart,
 		OpSocketReadMetadata,
 		OpSocketRead,
 		OpFileWrite,
@@ -86,7 +85,7 @@ FRecorderRelay::FRecorderRelay(asio::ip::tcp::socket& Socket, FStore& InStore)
 
 	if (CreateTrace())
 	{
-		OnIoComplete(OpStart, 0);
+		OnIoComplete(OpFileWrite, 0);
 	}
 }
 
@@ -248,7 +247,6 @@ void FRecorderRelay::OnIoComplete(uint32 Id, int32 Size)
 		Output->Write(Buffer, Size, this, OpFileWrite);
 		break;
 
-	case OpStart:
 	case OpFileWrite:
 		Input.ReadSome(Buffer, BufferSize, this, ActiveReadOp);
 		break;
