Sanitise received payload size to within sensible limits.

Any random data that might be received by the CBOR server could cause it to allocate a lot of memory, or suffer a signed overflow and collapse. This can happen for example if a machine is port scanned by an IT department. To guard against this, the input is checked to ensure it is within a suitable before use.

#rnx

diff --git a/Engine/Source/Programs/UnrealTraceServer/src/StoreCborServer.cpp b/Engine/Source/Programs/UnrealTraceServer/src/StoreCborServer.cpp
--- a/Engine/Source/Programs/UnrealTraceServer/src/StoreCborServer.cpp
+++ b/Engine/Source/Programs/UnrealTraceServer/src/StoreCborServer.cpp
@@ -281,6 +281,11 @@ void FStoreCborPeer::OnIoComplete(uint32 Id, int32 Size)
 		break;
 
 	case OpReadPayloadSize:
+		if (PayloadSize == 0 || PayloadSize > 1024)
+		{
+			SendError(EStatusCode::BadRequest);
+			break;
+		}
 		Socket.Read(Response.Reserve(PayloadSize), PayloadSize, this, OpReadPayload);
 		break;
 
