Only create traces with the store when valid trace data has been detected.

Prior to this change, a trace was registered with the store component before creating the relay that receives data over the socket and writes it to a file. While the relay does check the preamble of the data received, and rejects anything that is not a trace, it does not clean up the store. Now no trace is created until the relay has determined it has a trace to receive data for. This stops zero-sized traces littering a users disk which can happen if the machine is port scanned by their IT department.

#rnx

diff --git a/Engine/Source/Programs/UnrealTraceServer/src/Recorder.cpp b/Engine/Source/Programs/UnrealTraceServer/src/Recorder.cpp
--- a/Engine/Source/Programs/UnrealTraceServer/src/Recorder.cpp
+++ b/Engine/Source/Programs/UnrealTraceServer/src/Recorder.cpp
@@ -34,23 +34,30 @@ public:
 private:
 	virtual void		OnIoComplete(uint32 Id, int32 Size) override;
 	bool				CreateTrace();
-	bool				ReadMagic(uint32 Magic);
+	bool				ReadMagic();
 	bool				ReadMetadata(int32 Size);
 	static const uint32	BufferSize = 64 * 1024;
 	FAsioSocket			Input;
 	FAsioWriteable*		Output = nullptr;
 	FStore&				Store;
-	uint32				ActiveReadOp = OpSocketReadMetadata;
+	uint8*				PreambleCursor;
 	uint32				TraceId = 0;
 	uint16				ControlPort = 0;
 	uint8				Buffer[BufferSize];
 
 	enum
 	{
-		OpSocketReadMetadata,
+		OpMagicRead,
+		OpMetadataRead,
 		OpSocketRead,
 		OpFileWrite,
 	};
+
+	using MagicType				= uint32;
+	using MetadataSizeType		= uint16;
+	using VersionType			= struct { uint8 Transport; uint8 Protocol; };
+
+	static_assert(sizeof(VersionType) == 2, "Unexpected struct size");
 };
 
 ////////////////////////////////////////////////////////////////////////////////
@@ -83,10 +90,11 @@ FRecorderRelay::FRecorderRelay(asio::ip::tcp::socket& Socket, FStore& InStore)
 	);
 #endif
 
-	if (CreateTrace())
-	{
-		OnIoComplete(OpFileWrite, 0);
-	}
+	// Kick things off by reading the magic four bytes at the start of the stream
+	// along with an additional two bytes that are likely the metadata size.
+	uint32 PreambleReadSize = sizeof(MagicType) + sizeof(MetadataSizeType);
+	PreambleCursor = Buffer + PreambleReadSize;
+	Input.Read(Buffer, PreambleReadSize, this, OpMagicRead);
 }
 
 ////////////////////////////////////////////////////////////////////////////////
@@ -144,62 +152,72 @@ bool FRecorderRelay::CreateTrace()
 }
 
 ////////////////////////////////////////////////////////////////////////////////
-bool FRecorderRelay::ReadMagic(uint32 Magic)
+bool FRecorderRelay::ReadMagic()
 {
-	switch (Magic)
+	// Here we'll check the magic four bytes at the start of the stream and create
+	// a trace to write into if they are bytes we're expecting.
+
+	// We will only support clients that send the magic. Very early clients did
+	// not do this but they were unreleased and should no longer be in use.
+	if (Buffer[3] != 'T' || Buffer[2] != 'R' || Buffer[1] != 'C')
 	{
-	/* trace with metadata data */
-	case FourCc('T', 'R', 'C', '2'):
-		break;
+		return false;
+	}
 
-	/* valid, but to old or wrong endian for us */
-	case FourCc('T', 'R', 'C', 'E'):
-	case FourCc('E', 'C', 'R', 'T'):
-	case FourCc('2', 'C', 'R', 'T'):
-		return true;
+	// We can continue to support very old clients
+	if (Buffer[0] == 'E')
+	{
+		if (CreateTrace())
+		{
+			// Old clients have no metadata so we can go straight into the
+			// read-write loop. We've already got read data in Buffer.
+			Output->Write(Buffer, sizeof(MagicType), this, OpFileWrite);
+			return true;
+		}
+		return false;
+	}
+
+	// Later clients have a metadata block (TRC2). There's loose support for the
+	// future too if need be (TRC[3-9]).
+	if (Buffer[0] < '2' || Buffer[0] > '9')
+	{
+		return false;
+	}
 
-	/* unexpected magic */
-	default:
+	// Concatenate metadata into the buffer, first validating the given size is
+	// one that we can handle in a single read.
+	uint32 MetadataSize = *(MetadataSizeType*)(Buffer + sizeof(MagicType));
+	MetadataSize += sizeof(VersionType);
+	if (MetadataSize > BufferSize - uint32(ptrdiff_t(PreambleCursor - Buffer)))
+	{
 		return false;
 	}
 
-	return false;
+	Input.Read(PreambleCursor, MetadataSize, this, OpMetadataRead);
+
+	return true;
 }
 
 ////////////////////////////////////////////////////////////////////////////////
 bool FRecorderRelay::ReadMetadata(int32 Size)
 {
-	const uint8* Cursor = Buffer;
+	// At this point Buffer          [magic][md_size][metadata][t_ver][p_ver]
+	// looks like this;              Buffer--------->PreambleCursor--------->
+	//                                               |---------Size---------|
+
+	// We want to consume [metadata] so some adjustment is required.
+	int32 ReadSize = Size - sizeof(VersionType);
+	const uint8* Cursor = PreambleCursor;
 	auto Read = [&] (int32 SizeToRead)
 	{
 		const uint8* Ptr = Cursor;
 		Cursor += SizeToRead;
-		Size -= SizeToRead;
+		ReadSize -= SizeToRead;
 		return Ptr;
 	};
 
-	// Stream header
-	uint32 Magic;
-	if (Size < sizeof(Magic))
-	{
-		return true;
-	}
-
-	Magic = *(const uint32*)(Read(sizeof(Magic)));
-	if (!ReadMagic(Magic))
-	{
-		return false;
-	}
-
-	// MetadataSize field
-	if (Size < 2)
-	{
-		return true;
-	}
-	Size = *(const uint16*)(Read(2));
-
 	// MetadataFields
-	while (Size >= 2)
+	while (ReadSize > 0)
 	{
 		struct {
 			uint8	Size;
@@ -207,7 +225,7 @@ bool FRecorderRelay::ReadMetadata(int32 Size)
 		} MetadataField;
 		MetadataField = *(const decltype(MetadataField)*)(Read(sizeof(MetadataField)));
 
-		if (Size < MetadataField.Size)
+		if (ReadSize < MetadataField.Size)
 		{
 			break;
 		}
@@ -217,9 +235,19 @@ bool FRecorderRelay::ReadMetadata(int32 Size)
 			ControlPort = *(const uint16*)Cursor;
 		}
 
-		Size -= MetadataField.Size;
+		ReadSize -= MetadataField.Size;
 	}
 
+	// Now we've a full preamble we are ready to write the trace.
+	if (!CreateTrace())
+	{
+		return false;
+	}
+
+	// Analysis needs the preamble too.
+	uint32 PreambleSize = uint32(ptrdiff_t(PreambleCursor - Buffer)) + Size;
+	Output->Write(Buffer, PreambleSize, this, OpFileWrite);
+
 	return true;
 }
 
@@ -234,21 +262,26 @@ void FRecorderRelay::OnIoComplete(uint32 Id, int32 Size)
 
 	switch (Id)
 	{
-	case OpSocketReadMetadata:
-		ActiveReadOp = OpSocketRead;
+	case OpMagicRead:
+		if (!ReadMagic())
+		{
+			Close();
+		}
+		break;
+
+	case OpMetadataRead:
 		if (!ReadMetadata(Size))
 		{
 			Close();
-			return;
 		}
-		/* fallthrough */
+		break;
 
 	case OpSocketRead:
 		Output->Write(Buffer, Size, this, OpFileWrite);
 		break;
 
 	case OpFileWrite:
-		Input.ReadSome(Buffer, BufferSize, this, ActiveReadOp);
+		Input.ReadSome(Buffer, BufferSize, this, OpSocketRead);
 		break;
 	}
 }
