Removed unused function
#rnx
(git:7905ddc - Martin Ridgers - [REDACTED_PHONE]:14:26 +0100)

diff --git a/Engine/Source/Programs/UnrealTraceServer/src/Recorder.cpp b/Engine/Source/Programs/UnrealTraceServer/src/Recorder.cpp
--- a/Engine/Source/Programs/UnrealTraceServer/src/Recorder.cpp
+++ b/Engine/Source/Programs/UnrealTraceServer/src/Recorder.cpp
@@ -10,14 +10,6 @@
 #	include <Mstcpip.h>
 #endif
 
-////////////////////////////////////////////////////////////////////////////////
-constexpr uint32 FourCc(int A, int B, int C, int D)
-{
-    return ((A << 24) | (B << 16) | (C << 8) | D);
-}
-
-
-
 ////////////////////////////////////////////////////////////////////////////////
 class FRecorderRelay
 	: public FAsioIoSink
