Also fail on unexpected metadata sizes too
#rnx
(git:b9bd36e - Martin Ridgers - [REDACTED_PHONE]:19:11 +0100)

diff --git a/Engine/Source/Programs/UnrealTraceServer/src/Recorder.cpp b/Engine/Source/Programs/UnrealTraceServer/src/Recorder.cpp
--- a/Engine/Source/Programs/UnrealTraceServer/src/Recorder.cpp
+++ b/Engine/Source/Programs/UnrealTraceServer/src/Recorder.cpp
@@ -219,7 +219,7 @@ bool FRecorderRelay::ReadMetadata(int32 Size)
 
 		if (ReadSize < MetadataField.Size)
 		{
-			break;
+			return false;
 		}
 
 		if (MetadataField.Id == 0) /* ControlPortFieldId */
