Reject traces without well-formed metadata blobs
#rnx
(git:37fcb57 - Martin Ridgers - [REDACTED_PHONE]:21:48 +0100)

diff --git a/Engine/Source/Programs/UnrealTraceServer/src/Recorder.cpp b/Engine/Source/Programs/UnrealTraceServer/src/Recorder.cpp
--- a/Engine/Source/Programs/UnrealTraceServer/src/Recorder.cpp
+++ b/Engine/Source/Programs/UnrealTraceServer/src/Recorder.cpp
@@ -230,6 +230,12 @@ bool FRecorderRelay::ReadMetadata(int32 Size)
 		ReadSize -= MetadataField.Size;
 	}
 
+	// There should be no data left to consume if the metadata was well-formed
+	if (ReadSize != 0)
+	{
+		return false;
+	}
+
 	// Now we've a full preamble we are ready to write the trace.
 	if (!CreateTrace())
 	{
