Updated UnrealTraceServer

Traces aren't created until a valid stream has been identified. Guard against bad values in CBOR payloads.

#rnx

diff --git a/Engine/Source/Programs/UnrealTraceServer/src/Version.h b/Engine/Source/Programs/UnrealTraceServer/src/Version.h
--- a/Engine/Source/Programs/UnrealTraceServer/src/Version.h
+++ b/Engine/Source/Programs/UnrealTraceServer/src/Version.h
@@ -7,7 +7,7 @@
 #endif
 
 #if !defined(TS_VERSION_MINOR)
-#	define TS_VERSION_MINOR		12
+#	define TS_VERSION_MINOR		13
 #endif
 
 #define TS_ICON_ID				493
