Recorder's parsing of metadata wasn't updating the read cursor correctly
#rnx
(git:c8d38dc - Martin Ridgers - [REDACTED_PHONE]:01:33 +0100)

diff --git a/Engine/Source/Programs/UnrealTraceServer/src/Recorder.cpp b/Engine/Source/Programs/UnrealTraceServer/src/Recorder.cpp
--- a/Engine/Source/Programs/UnrealTraceServer/src/Recorder.cpp
+++ b/Engine/Source/Programs/UnrealTraceServer/src/Recorder.cpp
@@ -200,16 +200,9 @@ bool FRecorderRelay::ReadMetadata(int32 Size)
 	// We want to consume [metadata] so some adjustment is required.
 	int32 ReadSize = Size - sizeof(VersionType);
 	const uint8* Cursor = PreambleCursor;
-	auto Read = [&] (int32 SizeToRead)
-	{
-		const uint8* Ptr = Cursor;
-		Cursor += SizeToRead;
-		ReadSize -= SizeToRead;
-		return Ptr;
-	};
 
 	// MetadataFields
-	while (ReadSize > 0)
+	while (ReadSize >= 2)
 	{
 		struct {
 			uint8	Size;
@@ -217,6 +210,9 @@ bool FRecorderRelay::ReadMetadata(int32 Size)
 		} MetadataField;
 		MetadataField = *(const decltype(MetadataField)*)(Read(sizeof(MetadataField)));
 
+		Cursor += sizeof(FMetadataHeader);
+		ReadSize -= sizeof(FMetadataHeader);
+
 		if (ReadSize < MetadataField.Size)
 		{
 			return false;
@@ -227,6 +223,7 @@ bool FRecorderRelay::ReadMetadata(int32 Size)
 			ControlPort = *(const uint16*)Cursor;
 		}
 
+		Cursor += MetadataField.Size;
 		ReadSize -= MetadataField.Size;
 	}
 
