Old-school (and almost certainly unused) was missing version bytes
#rnx
(git:3bb3e2c - Martin Ridgers - [REDACTED_PHONE]:50:08 +0100)

diff --git a/Engine/Source/Programs/UnrealTraceServer/src/Recorder.cpp b/Engine/Source/Programs/UnrealTraceServer/src/Recorder.cpp
--- a/Engine/Source/Programs/UnrealTraceServer/src/Recorder.cpp
+++ b/Engine/Source/Programs/UnrealTraceServer/src/Recorder.cpp
@@ -163,7 +163,7 @@ bool FRecorderRelay::ReadMagic()
 		{
 			// Old clients have no metadata so we can go straight into the
 			// read-write loop. We've already got read data in Buffer.
-			Output->Write(Buffer, sizeof(MagicType), this, OpFileWrite);
+			Output->Write(Buffer, sizeof(MagicType) + sizeof(MetadataSizeType), this, OpFileWrite);
 			return true;
 		}
 		return false;
