Remove use of decltype()
#rnx
(git:28cd2cc - Martin Ridgers - [REDACTED_PHONE]:59:54 +0100)

diff --git a/Engine/Source/Programs/UnrealTraceServer/src/Recorder.cpp b/Engine/Source/Programs/UnrealTraceServer/src/Recorder.cpp
--- a/Engine/Source/Programs/UnrealTraceServer/src/Recorder.cpp
+++ b/Engine/Source/Programs/UnrealTraceServer/src/Recorder.cpp
@@ -204,11 +204,12 @@ bool FRecorderRelay::ReadMetadata(int32 Size)
 	// MetadataFields
 	while (ReadSize >= 2)
 	{
-		struct {
+		struct FMetadataHeader
+		{
 			uint8	Size;
 			uint8	Id;
-		} MetadataField;
-		MetadataField = *(const decltype(MetadataField)*)(Read(sizeof(MetadataField)));
+		};
+		const auto& MetadataField = *(FMetadataHeader*)(Cursor);
 
 		Cursor += sizeof(FMetadataHeader);
 		ReadSize -= sizeof(FMetadataHeader);
