Add asio::io_context accessor to FAsioFile similar to FAsioSocket
#rnx
(git:ebf34ee - Martin Ridgers - [REDACTED_PHONE]:57:54 +0100)

diff --git a/Engine/Source/Programs/UnrealTraceServer/src/AsioFile.cpp b/Engine/Source/Programs/UnrealTraceServer/src/AsioFile.cpp
--- a/Engine/Source/Programs/UnrealTraceServer/src/AsioFile.cpp
+++ b/Engine/Source/Programs/UnrealTraceServer/src/AsioFile.cpp
@@ -19,6 +19,18 @@ FAsioFile::FAsioFile(asio::io_context& IoContext, uintptr_t OsHandle)
 {
 }
 
+////////////////////////////////////////////////////////////////////////////////
+asio::io_context& FAsioFile::GetIoContext()
+{
+#if TS_USING(TS_PLATFORM_WINDOWS)
+	auto& Inner = Handle;
+#else
+	auto& Inner = StreamDescriptor;
+#endif
+
+	return Inner.get_executor().context();
+}
+
 ////////////////////////////////////////////////////////////////////////////////
 bool FAsioFile::IsOpen() const
 {
diff --git a/Engine/Source/Programs/UnrealTraceServer/src/AsioFile.h b/Engine/Source/Programs/UnrealTraceServer/src/AsioFile.h
--- a/Engine/Source/Programs/UnrealTraceServer/src/AsioFile.h
+++ b/Engine/Source/Programs/UnrealTraceServer/src/AsioFile.h
@@ -17,6 +17,7 @@ class FAsioFile
 {
 public:
 										FAsioFile(asio::io_context& IoContext, uintptr_t OsHandle);
+	asio::io_context&					GetIoContext();
 	static FAsioWriteable*				WriteFile(asio::io_context& IoContext, const FPath& Path);
 	static FAsioReadable*				ReadFile(asio::io_context& IoContext, const FPath& Path);
 	virtual bool						IsOpen() const override;
