Gave metadata ops a higher ID value
#rnx
(git:ea3dfe1 - Martin Ridgers - [REDACTED_PHONE]:18:42 +0100)

diff --git a/Engine/Source/Programs/UnrealTraceServer/src/Recorder.cpp b/Engine/Source/Programs/UnrealTraceServer/src/Recorder.cpp
--- a/Engine/Source/Programs/UnrealTraceServer/src/Recorder.cpp
+++ b/Engine/Source/Programs/UnrealTraceServer/src/Recorder.cpp
@@ -39,8 +39,8 @@ private:
 
 	enum
 	{
-		OpMagicRead,
-		OpMetadataRead,
+		OpMagicRead			= 0x1000,
+		OpMetadataRead		= 0x1001,
 		OpSocketRead,
 		OpFileWrite,
 	};
