Added a lawn like the rest of the file
#rnx
(git:8da762d - Martin Ridgers - [REDACTED_PHONE]:19:48 +0100)

diff --git a/Engine/Source/Programs/UnrealTraceServer/src/Store.cpp b/Engine/Source/Programs/UnrealTraceServer/src/Store.cpp
--- a/Engine/Source/Programs/UnrealTraceServer/src/Store.cpp
+++ b/Engine/Source/Programs/UnrealTraceServer/src/Store.cpp
@@ -31,6 +31,8 @@ class FStore::FDirWatcher
 public:
 	using asio::windows::object_handle::object_handle;
 };
+
+////////////////////////////////////////////////////////////////////////////////
 #elif TS_USING(TS_PLATFORM_LINUX)
 class FStore::FDirWatcher
 	: public asio::posix::stream_descriptor
@@ -43,6 +45,8 @@ public:
 		asio::posix::stream_descriptor::async_wait(asio::posix::stream_descriptor::wait_read, InHandler);
 	}
 };
+
+////////////////////////////////////////////////////////////////////////////////
 #elif TS_USING(TS_PLATFORM_MAC)
 class FStore::FDirWatcher
 {
@@ -83,6 +87,7 @@ private:
 	FPath StoreDir;
 };
 
+////////////////////////////////////////////////////////////////////////////////
 void MacCallback(ConstFSEventStreamRef StreamRef,
 					void* InDirWatcherPtr,
 					size_t EventCount,
@@ -97,6 +102,7 @@ void MacCallback(ConstFSEventStreamRef StreamRef,
 	DirWatcherPtr->ProcessChanges(EventCount, EventPaths, EventFlags);
 }
 
+////////////////////////////////////////////////////////////////////////////////
 void FStore::FDirWatcher::async_wait(HandlerType InHandler)
 {
 	if (bIsRunning)
