Callback doesn't need external linkage
#rnx
(git:4470900 - Martin Ridgers - [REDACTED_PHONE]:20:09 +0100)

diff --git a/Engine/Source/Programs/UnrealTraceServer/src/Store.cpp b/Engine/Source/Programs/UnrealTraceServer/src/Store.cpp
--- a/Engine/Source/Programs/UnrealTraceServer/src/Store.cpp
+++ b/Engine/Source/Programs/UnrealTraceServer/src/Store.cpp
@@ -88,7 +88,7 @@ private:
 };
 
 ////////////////////////////////////////////////////////////////////////////////
-void MacCallback(ConstFSEventStreamRef StreamRef,
+static void MacCallback(ConstFSEventStreamRef StreamRef,
 					void* InDirWatcherPtr,
 					size_t EventCount,
 					void* EventPaths,
