Use the heap for the recorder's buffer instead of the stack
#rnx
(git:440fbc0 - Martin Ridgers - [REDACTED_PHONE]:22:28 +0100)

diff --git a/Engine/Source/Programs/UnrealTraceServer/src/Recorder.cpp b/Engine/Source/Programs/UnrealTraceServer/src/Recorder.cpp
--- a/Engine/Source/Programs/UnrealTraceServer/src/Recorder.cpp
+++ b/Engine/Source/Programs/UnrealTraceServer/src/Recorder.cpp
@@ -35,7 +35,7 @@ private:
 	uint8*				PreambleCursor;
 	uint32				TraceId = 0;
 	uint16				ControlPort = 0;
-	uint8				Buffer[BufferSize];
+	uint8*				Buffer;
 
 	enum
 	{
@@ -57,6 +57,8 @@ FRecorderRelay::FRecorderRelay(asio::ip::tcp::socket& Socket, FStore& InStore)
 : Input(Socket)
 , Store(InStore)
 {
+	Buffer = new uint8[BufferSize];
+
 #if TS_USING(TS_PLATFORM_WINDOWS)
 	// Trace data is a stream and communication is one way. It is implemented
 	// this way to share code between sending trace data over the wire and writing
@@ -98,6 +100,8 @@ FRecorderRelay::~FRecorderRelay()
 		check(!Output->IsOpen());
 		delete Output;
 	}
+
+	delete[] Buffer;
 }
 
 ////////////////////////////////////////////////////////////////////////////////
