Bumped recorder's buffer size
#rnx
(git:b4b105d - Martin Ridgers - [REDACTED_PHONE]:23:16 +0100)

diff --git a/Engine/Source/Programs/UnrealTraceServer/src/Recorder.cpp b/Engine/Source/Programs/UnrealTraceServer/src/Recorder.cpp
--- a/Engine/Source/Programs/UnrealTraceServer/src/Recorder.cpp
+++ b/Engine/Source/Programs/UnrealTraceServer/src/Recorder.cpp
@@ -28,7 +28,7 @@ private:
 	bool				CreateTrace();
 	bool				ReadMagic();
 	bool				ReadMetadata(int32 Size);
-	static const uint32	BufferSize = 64 * 1024;
+	static const uint32	BufferSize = 256 * 1024;
 	FAsioSocket			Input;
 	FAsioWriteable*		Output = nullptr;
 	FStore&				Store;
