Small semantic adjustment to how the magic is parsed to reduce use of 'Buffer' directly
#rnx
(git:b393071 - Martin Ridgers - [REDACTED_PHONE]:36:56 +0100)

diff --git a/Engine/Source/Programs/UnrealTraceServer/src/Recorder.cpp b/Engine/Source/Programs/UnrealTraceServer/src/Recorder.cpp
--- a/Engine/Source/Programs/UnrealTraceServer/src/Recorder.cpp
+++ b/Engine/Source/Programs/UnrealTraceServer/src/Recorder.cpp
@@ -150,24 +150,26 @@ bool FRecorderRelay::CreateTrace()
 ////////////////////////////////////////////////////////////////////////////////
 bool FRecorderRelay::ReadMagic()
 {
+	const uint8* Cursor = Buffer;
+
 	// Here we'll check the magic four bytes at the start of the stream and create
 	// a trace to write into if they are bytes we're expecting.
 
 	// We will only support clients that send the magic. Very early clients did
 	// not do this but they were unreleased and should no longer be in use.
-	if (Buffer[3] != 'T' || Buffer[2] != 'R' || Buffer[1] != 'C')
+	if (Cursor[3] != 'T' || Cursor[2] != 'R' || Cursor[1] != 'C')
 	{
 		return false;
 	}
 
 	// We can continue to support very old clients
-	if (Buffer[0] == 'E')
+	if (Cursor[0] == 'E')
 	{
 		if (CreateTrace())
 		{
 			// Old clients have no metadata so we can go straight into the
 			// read-write loop. We've already got read data in Buffer.
-			Output->Write(Buffer, sizeof(MagicType) + sizeof(MetadataSizeType), this, OpFileWrite);
+			Output->Write(Cursor, sizeof(MagicType) + sizeof(MetadataSizeType), this, OpFileWrite);
 			return true;
 		}
 		return false;
@@ -175,16 +177,16 @@ bool FRecorderRelay::ReadMagic()
 
 	// Later clients have a metadata block (TRC2). There's loose support for the
 	// future too if need be (TRC[3-9]).
-	if (Buffer[0] < '2' || Buffer[0] > '9')
+	if (Cursor[0] < '2' || Cursor[0] > '9')
 	{
 		return false;
 	}
 
 	// Concatenate metadata into the buffer, first validating the given size is
 	// one that we can handle in a single read.
-	uint32 MetadataSize = *(MetadataSizeType*)(Buffer + sizeof(MagicType));
+	uint32 MetadataSize = *(MetadataSizeType*)(Cursor + sizeof(MagicType));
 	MetadataSize += sizeof(VersionType);
-	if (MetadataSize > BufferSize - uint32(ptrdiff_t(PreambleCursor - Buffer)))
+	if (MetadataSize > BufferSize - uint32(ptrdiff_t(PreambleCursor - Cursor)))
 	{
 		return false;
 	}
