Let recorder's OnIoComplete() handle what to do next
#rnx
(git:8565abd - Martin Ridgers - [REDACTED_PHONE]:51:02 +0100)

diff --git a/Engine/Source/Programs/UnrealTraceServer/src/Recorder.cpp b/Engine/Source/Programs/UnrealTraceServer/src/Recorder.cpp
--- a/Engine/Source/Programs/UnrealTraceServer/src/Recorder.cpp
+++ b/Engine/Source/Programs/UnrealTraceServer/src/Recorder.cpp
@@ -169,7 +169,7 @@ bool FRecorderRelay::ReadMagic()
 		{
 			// Old clients have no metadata so we can go straight into the
 			// read-write loop. We've already got read data in Buffer.
-			Output->Write(Cursor, sizeof(MagicType) + sizeof(MetadataSizeType), this, OpFileWrite);
+			OnIoComplete(OpSocketRead, sizeof(MagicType) + sizeof(MetadataSizeType));
 			return true;
 		}
 		return false;
@@ -248,7 +248,7 @@ bool FRecorderRelay::ReadMetadata(int32 Size)
 
 	// Analysis needs the preamble too.
 	uint32 PreambleSize = uint32(ptrdiff_t(PreambleCursor - Buffer)) + Size;
-	Output->Write(Buffer, PreambleSize, this, OpFileWrite);
+	OnIoComplete(OpSocketRead, PreambleSize);
 
 	return true;
 }
