It is sufficient to return directly in each case here
#rnx
(git:3fa27cc - Martin Ridgers - [REDACTED_PHONE]:54:17 +0100)

diff --git a/Engine/Source/Programs/UnrealTraceServer/src/Recorder.cpp b/Engine/Source/Programs/UnrealTraceServer/src/Recorder.cpp
--- a/Engine/Source/Programs/UnrealTraceServer/src/Recorder.cpp
+++ b/Engine/Source/Programs/UnrealTraceServer/src/Recorder.cpp
@@ -269,22 +269,22 @@ void FRecorderRelay::OnIoComplete(uint32 Id, int32 Size)
 		{
 			Close();
 		}
-		break;
+		return;
 
 	case OpMetadataRead:
 		if (!ReadMetadata(Size))
 		{
 			Close();
 		}
-		break;
+		return;
 
 	case OpSocketRead:
 		Output->Write(Buffer, Size, this, OpFileWrite);
-		break;
+		return;
 
 	case OpFileWrite:
 		Input.ReadSome(Buffer, BufferSize, this, OpSocketRead);
-		break;
+		return;
 	}
 }
 
