Add double-buffering to the recorder allowing reads and writes to overlap each other.
#rnx
(git:7568f24 - Martin Ridgers - [REDACTED_PHONE]:55:17 +0100)

diff --git a/Engine/Source/Programs/UnrealTraceServer/src/Recorder.cpp b/Engine/Source/Programs/UnrealTraceServer/src/Recorder.cpp
--- a/Engine/Source/Programs/UnrealTraceServer/src/Recorder.cpp
+++ b/Engine/Source/Programs/UnrealTraceServer/src/Recorder.cpp
@@ -35,14 +35,19 @@ private:
 	uint8*				PreambleCursor;
 	uint32				TraceId = 0;
 	uint16				ControlPort = 0;
-	uint8*				Buffer;
+	uint8				FillDrainCounter = 1;
+	uint32				BufferPolarity = 0;
+	uint8*				Buffer[2];
+	uint32				FillSize[2] = {};
 
 	enum
 	{
 		OpMagicRead			= 0x1000,
 		OpMetadataRead		= 0x1001,
-		OpSocketRead,
-		OpFileWrite,
+		OpBuffer0			= 0b00,
+		OpBuffer1			= 0b10,
+		OpFill				= 0b00,
+		OpDrain				= 0b01,
 	};
 
 	using MagicType				= uint32;
@@ -57,7 +62,8 @@ FRecorderRelay::FRecorderRelay(asio::ip::tcp::socket& Socket, FStore& InStore)
 : Input(Socket)
 , Store(InStore)
 {
-	Buffer = new uint8[BufferSize];
+	Buffer[0] = new uint8[BufferSize];
+	Buffer[1] = new uint8[BufferSize];
 
 #if TS_USING(TS_PLATFORM_WINDOWS)
 	// Trace data is a stream and communication is one way. It is implemented
@@ -87,8 +93,8 @@ FRecorderRelay::FRecorderRelay(asio::ip::tcp::socket& Socket, FStore& InStore)
 	// Kick things off by reading the magic four bytes at the start of the stream
 	// along with an additional two bytes that are likely the metadata size.
 	uint32 PreambleReadSize = sizeof(MagicType) + sizeof(MetadataSizeType);
-	PreambleCursor = Buffer + PreambleReadSize;
-	Input.Read(Buffer, PreambleReadSize, this, OpMagicRead);
+	PreambleCursor = Buffer[0] + PreambleReadSize;
+	Input.Read(Buffer[0], PreambleReadSize, this, OpMagicRead);
 }
 
 ////////////////////////////////////////////////////////////////////////////////
@@ -101,7 +107,8 @@ FRecorderRelay::~FRecorderRelay()
 		delete Output;
 	}
 
-	delete[] Buffer;
+	delete[] Buffer[1];
+	delete[] Buffer[0];
 }
 
 ////////////////////////////////////////////////////////////////////////////////
@@ -150,7 +157,7 @@ bool FRecorderRelay::CreateTrace()
 ////////////////////////////////////////////////////////////////////////////////
 bool FRecorderRelay::ReadMagic()
 {
-	const uint8* Cursor = Buffer;
+	const uint8* Cursor = Buffer[0];
 
 	// Here we'll check the magic four bytes at the start of the stream and create
 	// a trace to write into if they are bytes we're expecting.
@@ -169,7 +176,7 @@ bool FRecorderRelay::ReadMagic()
 		{
 			// Old clients have no metadata so we can go straight into the
 			// read-write loop. We've already got read data in Buffer.
-			OnIoComplete(OpSocketRead, sizeof(MagicType) + sizeof(MetadataSizeType));
+			OnIoComplete(OpBuffer0|OpFill, sizeof(MagicType) + sizeof(MetadataSizeType));
 			return true;
 		}
 		return false;
@@ -247,8 +254,8 @@ bool FRecorderRelay::ReadMetadata(int32 Size)
 	}
 
 	// Analysis needs the preamble too.
-	uint32 PreambleSize = uint32(ptrdiff_t(PreambleCursor - Buffer)) + Size;
-	OnIoComplete(OpSocketRead, PreambleSize);
+	uint32 PreambleSize = uint32(ptrdiff_t(PreambleCursor - Buffer[0])) + Size;
+	OnIoComplete(OpBuffer0|OpFill, PreambleSize);
 
 	return true;
 }
@@ -262,6 +269,7 @@ void FRecorderRelay::OnIoComplete(uint32 Id, int32 Size)
 		return;
 	}
 
+	// A completed preamble read?
 	switch (Id)
 	{
 	case OpMagicRead:
@@ -277,15 +285,32 @@ void FRecorderRelay::OnIoComplete(uint32 Id, int32 Size)
 			Close();
 		}
 		return;
+	}
 
-	case OpSocketRead:
-		Output->Write(Buffer, Size, this, OpFileWrite);
-		return;
+	// If we've got to here then a fill or drain op has completed.
+
+	// Cache fill sizes so they can be used when we are ready to issue the drain
+	if ((Id & 0b01) == 0)
+	{
+		uint32 BufferIndex = (Id & 0b10) >> 1;
+		FillSize[BufferIndex] = Size;
+	}
 
-	case OpFileWrite:
-		Input.ReadSome(Buffer, BufferSize, this, OpSocketRead);
+	// We only dispatch another fill/drain pair once the previous two have
+	// completed. This is so we don't overlap fills, drains, or buffer use.
+	++FillDrainCounter;
+	if (FillDrainCounter < 2)
+	{
 		return;
 	}
+	FillDrainCounter = 0;
+
+	// At this point we've one buffer that's been filled and another that has
+	// finished being drained. We are free to issue two more concurrent ops
+	uint32 BufferIndex = (BufferPolarity != 0);
+	Output->Write(Buffer[BufferIndex], FillSize[BufferIndex], this, BufferPolarity|OpDrain);
+	BufferPolarity ^= 0b10;
+	Input.ReadSome(Buffer[BufferIndex ^ 1], BufferSize, this, BufferPolarity|OpFill);
 }
 
 
