Recording relays should also be considered open when a drain op is in flight
#rnx
(git:6eb7f3c - Martin Ridgers - [REDACTED_PHONE]:37:46 +0100)

diff --git a/Engine/Source/Programs/UnrealTraceServer/src/Recorder.cpp b/Engine/Source/Programs/UnrealTraceServer/src/Recorder.cpp
--- a/Engine/Source/Programs/UnrealTraceServer/src/Recorder.cpp
+++ b/Engine/Source/Programs/UnrealTraceServer/src/Recorder.cpp
@@ -114,7 +114,11 @@ FRecorderRelay::~FRecorderRelay()
 ////////////////////////////////////////////////////////////////////////////////
 bool FRecorderRelay::IsOpen()
 {
-	return Input.IsOpen();
+	// Even if the input socket has been closed we should still report ourselves
+	// as open if there is a pending write on the output.
+	bool Ret = (Output != nullptr && FillDrainCounter < 2);
+	Ret |= Input.IsOpen();
+	return Ret;
 }
 
 ////////////////////////////////////////////////////////////////////////////////
