Fill ops should still update the counter even if there's an error
#rnx
(git:9c83877 - Martin Ridgers - [REDACTED_PHONE]:38:53 +0100)

diff --git a/Engine/Source/Programs/UnrealTraceServer/src/Recorder.cpp b/Engine/Source/Programs/UnrealTraceServer/src/Recorder.cpp
--- a/Engine/Source/Programs/UnrealTraceServer/src/Recorder.cpp
+++ b/Engine/Source/Programs/UnrealTraceServer/src/Recorder.cpp
@@ -269,6 +269,7 @@ void FRecorderRelay::OnIoComplete(uint32 Id, int32 Size)
 {
 	if (Size < 0)
 	{
+		FillDrainCounter += (Output != nullptr);
 		Close();
 		return;
 	}
