Added a Last() method
#rnx
(git:77ccbef - Martin Ridgers - [REDACTED_PHONE]:32:23 +0100)

diff --git a/Engine/Source/Programs/UnrealTraceServer/src/Foundation.h b/Engine/Source/Programs/UnrealTraceServer/src/Foundation.h
--- a/Engine/Source/Programs/UnrealTraceServer/src/Foundation.h
+++ b/Engine/Source/Programs/UnrealTraceServer/src/Foundation.h
@@ -25,6 +25,7 @@ struct TArray
 	void		Empty()								{ return Super::clear(); }
 	void		SetNumUninitialized(int Num)		{ return Super::resize(Num); }
 	void		Append(const void* Data, int Num)	{ Super::insert(Super::end(), (const Type*)Data, ((const Type*)Data) + Num); }
+	Type&		Last()								{ return Super::back(); }
 };
 
 ////////////////////////////////////////////////////////////////////////////////
