Accessor to query if an ASIO object has data available
#rnx
(git:08a245b - Martin Ridgers - [REDACTED_PHONE]:34:34 +0100)

diff --git a/Engine/Source/Programs/UnrealTraceServer/src/AsioFile.cpp b/Engine/Source/Programs/UnrealTraceServer/src/AsioFile.cpp
--- a/Engine/Source/Programs/UnrealTraceServer/src/AsioFile.cpp
+++ b/Engine/Source/Programs/UnrealTraceServer/src/AsioFile.cpp
@@ -51,6 +51,25 @@ void FAsioFile::Close()
 #endif
 }
 
+////////////////////////////////////////////////////////////////////////////////
+bool FAsioFile::HasDataAvailable() const
+{
+	uint64 FileSize = 0;
+#if TS_USING(TS_PLATFORM_WINDOWS)
+	LARGE_INTEGER Out;
+	HANDLE Inner = const_cast<asio::windows::random_access_handle&>(Handle).native_handle();
+	GetFileSizeEx(Inner, &Out);
+	FileSize = Out.QuadPart;
+#else
+	struct stat Stat;
+	int32 Inner = const_cast<asio::posix::stream_descriptor&>(StreamDescriptor).native_handle();
+	fstat(Inner, &Stat);
+	FileSize = Stat.st_size;
+	static_assert(sizeof(Stat.st_size) >= sizeof(FileSize), "fstat() reports sizes that are too small");
+#endif
+	return Offset < FileSize;
+}
+
 ////////////////////////////////////////////////////////////////////////////////
 bool FAsioFile::Write(const void* Src, uint32 Size, FAsioIoSink* Sink, uint32 Id)
 {
diff --git a/Engine/Source/Programs/UnrealTraceServer/src/AsioFile.h b/Engine/Source/Programs/UnrealTraceServer/src/AsioFile.h
--- a/Engine/Source/Programs/UnrealTraceServer/src/AsioFile.h
+++ b/Engine/Source/Programs/UnrealTraceServer/src/AsioFile.h
@@ -22,6 +22,7 @@ public:
 	static FAsioReadable*				ReadFile(asio::io_context& IoContext, const FPath& Path);
 	virtual bool						IsOpen() const override;
 	virtual void						Close() override;
+	virtual bool						HasDataAvailable() const override;
 	virtual bool						Write(const void* Src, uint32 Size, FAsioIoSink* Sink, uint32 Id) override;
 	virtual bool						Read(void* Dest, uint32 Size, FAsioIoSink* Sink, uint32 Id) override;
 	virtual bool						ReadSome(void* Dest, uint32 BufferSize, FAsioIoSink* Sink, uint32 Id) override;
diff --git a/Engine/Source/Programs/UnrealTraceServer/src/AsioIoable.h b/Engine/Source/Programs/UnrealTraceServer/src/AsioIoable.h
--- a/Engine/Source/Programs/UnrealTraceServer/src/AsioIoable.h
+++ b/Engine/Source/Programs/UnrealTraceServer/src/AsioIoable.h
@@ -33,6 +33,7 @@ class FAsioReadable
 	: public virtual FAsioIoable
 {
 public:
+	virtual bool	HasDataAvailable() const = 0;
 	virtual bool	Read(void* Dest, uint32 Size, FAsioIoSink* Sink, uint32 Id) = 0;
 	virtual bool	ReadSome(void* Dest, uint32 DestSize, FAsioIoSink* Sink, uint32 Id) = 0;
 };
diff --git a/Engine/Source/Programs/UnrealTraceServer/src/AsioSocket.cpp b/Engine/Source/Programs/UnrealTraceServer/src/AsioSocket.cpp
--- a/Engine/Source/Programs/UnrealTraceServer/src/AsioSocket.cpp
+++ b/Engine/Source/Programs/UnrealTraceServer/src/AsioSocket.cpp
@@ -28,6 +28,12 @@ bool FAsioSocket::IsOpen() const
 	return Socket.is_open();
 }
 
+////////////////////////////////////////////////////////////////////////////////
+bool FAsioSocket::HasDataAvailable() const
+{
+	return (Socket.available() > 0);
+}
+
 ////////////////////////////////////////////////////////////////////////////////
 uint32 FAsioSocket::GetRemoteAddress() const
 {
diff --git a/Engine/Source/Programs/UnrealTraceServer/src/AsioSocket.h b/Engine/Source/Programs/UnrealTraceServer/src/AsioSocket.h
--- a/Engine/Source/Programs/UnrealTraceServer/src/AsioSocket.h
+++ b/Engine/Source/Programs/UnrealTraceServer/src/AsioSocket.h
@@ -16,6 +16,7 @@ public:
 	asio::io_context&		GetIoContext();
 	virtual bool			IsOpen() const override;
 	virtual void			Close() override;
+	virtual bool			HasDataAvailable() const override;
 	uint32					GetRemoteAddress() const;
 	uint32					GetRemotePort() const;
 	uint32					GetLocalPort() const;
