Tidied up some of the per-platform preprocessing
#rnx
(git:6f3d642 - Martin Ridgers - [REDACTED_PHONE]:22:21 +0100)

diff --git a/Engine/Source/Programs/UnrealTraceServer/src/AsioFile.cpp b/Engine/Source/Programs/UnrealTraceServer/src/AsioFile.cpp
--- a/Engine/Source/Programs/UnrealTraceServer/src/AsioFile.cpp
+++ b/Engine/Source/Programs/UnrealTraceServer/src/AsioFile.cpp
@@ -11,58 +11,40 @@
 
 ////////////////////////////////////////////////////////////////////////////////
 FAsioFile::FAsioFile(asio::io_context& IoContext, uintptr_t OsHandle)
-#if TS_USING(TS_PLATFORM_WINDOWS)
-: Handle(IoContext, HANDLE(OsHandle))
-#else
-: StreamDescriptor(IoContext, OsHandle)
-#endif
+: Handle(IoContext, HandleType::native_handle_type(OsHandle))
 {
 }
 
 ////////////////////////////////////////////////////////////////////////////////
 asio::io_context& FAsioFile::GetIoContext()
 {
-#if TS_USING(TS_PLATFORM_WINDOWS)
-	auto& Inner = Handle;
-#else
-	auto& Inner = StreamDescriptor;
-#endif
-
-	return Inner.get_executor().context();
+	return Handle.get_executor().context();
 }
 
 ////////////////////////////////////////////////////////////////////////////////
 bool FAsioFile::IsOpen() const
 {
-#if TS_USING(TS_PLATFORM_WINDOWS)
 	return Handle.is_open();
-#else
-	return StreamDescriptor.is_open();
-#endif
 }
 
 ////////////////////////////////////////////////////////////////////////////////
 void FAsioFile::Close()
 {
-#if TS_USING(TS_PLATFORM_WINDOWS)
 	Handle.close();
-#else
-	StreamDescriptor.close();
-#endif
 }
 
 ////////////////////////////////////////////////////////////////////////////////
 bool FAsioFile::HasDataAvailable() const
 {
+	auto Inner = const_cast<HandleType&>(Handle).native_handle();
+
 	uint64 FileSize = 0;
 #if TS_USING(TS_PLATFORM_WINDOWS)
 	LARGE_INTEGER Out;
-	HANDLE Inner = const_cast<asio::windows::random_access_handle&>(Handle).native_handle();
 	GetFileSizeEx(Inner, &Out);
 	FileSize = Out.QuadPart;
 #else
 	struct stat Stat;
-	int32 Inner = const_cast<asio::posix::stream_descriptor&>(StreamDescriptor).native_handle();
 	fstat(Inner, &Stat);
 	FileSize = Stat.st_size;
 	static_assert(sizeof(Stat.st_size) >= sizeof(FileSize), "fstat() reports sizes that are too small");
@@ -84,7 +66,7 @@ bool FAsioFile::Write(const void* Src, uint32 Size, FAsioIoSink* Sink, uint32 Id
 		Offset,
 #else
 	asio::async_write(
-		StreamDescriptor,
+		Handle,
 #endif
 		asio::buffer(Src, Size),
 		[this] (const asio::error_code& ErrorCode, size_t BytesWritten)
@@ -115,7 +97,7 @@ bool FAsioFile::ReadSome(void* Dest, uint32 DestSize, FAsioIoSink* Sink, uint32
 	Handle.async_read_some_at(
 		Offset,
 #else
-	StreamDescriptor.async_read_some(
+	Handle.async_read_some(
 #endif
 		asio::buffer(Dest, DestSize),
 		[this] (const asio::error_code& ErrorCode, size_t BytesRead)
diff --git a/Engine/Source/Programs/UnrealTraceServer/src/AsioFile.h b/Engine/Source/Programs/UnrealTraceServer/src/AsioFile.h
--- a/Engine/Source/Programs/UnrealTraceServer/src/AsioFile.h
+++ b/Engine/Source/Programs/UnrealTraceServer/src/AsioFile.h
@@ -28,11 +28,14 @@ public:
 	virtual bool						ReadSome(void* Dest, uint32 BufferSize, FAsioIoSink* Sink, uint32 Id) override;
 
 private:
+	using HandleType =
 #if TS_USING(TS_PLATFORM_WINDOWS)
-	asio::windows::random_access_handle	Handle;
+		asio::windows::random_access_handle;
 #else
-	asio::posix::stream_descriptor		StreamDescriptor;
+		asio::posix::stream_descriptor;
 #endif
+
+	HandleType							Handle;
 	uint64								Offset = 0;
 };
 
