Missing forward declare
#rnx
(git:346bbc9 - Martin Ridgers - [REDACTED_PHONE]:35:34 +0100)

diff --git a/Engine/Source/Programs/UnrealTraceServer/src/TraceRelay.h b/Engine/Source/Programs/UnrealTraceServer/src/TraceRelay.h
--- a/Engine/Source/Programs/UnrealTraceServer/src/TraceRelay.h
+++ b/Engine/Source/Programs/UnrealTraceServer/src/TraceRelay.h
@@ -9,6 +9,7 @@
 
 class FRecorder;
 class FAsioSocket;
+class FAsioReadable;
 
 ////////////////////////////////////////////////////////////////////////////////
 class FTraceRelay
