Also refresh traces when a directory is deleted (as POSIX does)
#rnx
(git:11b1225 - Martin Ridgers - [REDACTED_PHONE]:47:31 +0100)

diff --git a/Engine/Source/Programs/UnrealTraceServer/src/Store.cpp b/Engine/Source/Programs/UnrealTraceServer/src/Store.cpp
--- a/Engine/Source/Programs/UnrealTraceServer/src/Store.cpp
+++ b/Engine/Source/Programs/UnrealTraceServer/src/Store.cpp
@@ -308,7 +308,8 @@ FStore::FStore(asio::io_context& InIoContext, const FPath& InStoreDir)
 
 #if TS_USING(TS_PLATFORM_WINDOWS)
 	std::wstring StoreDirW = StoreDir;
-	HANDLE DirWatchHandle = FindFirstChangeNotificationW(StoreDirW.c_str(), false, FILE_NOTIFY_CHANGE_FILE_NAME);
+	HANDLE DirWatchHandle = FindFirstChangeNotificationW(StoreDirW.c_str(), false,
+		FILE_NOTIFY_CHANGE_FILE_NAME|FILE_NOTIFY_CHANGE_DIR_NAME);
 	if (DirWatchHandle == INVALID_HANDLE_VALUE)
 	{
 		DirWatchHandle = 0;
