Ditch support for old clients. As they use a different port they'd never connect anyway.
#rnx
(git:95ccf02 - Martin Ridgers - [REDACTED_PHONE]:33:42 +0100)

diff --git a/Engine/Source/Programs/UnrealTraceServer/src/Recorder.cpp b/Engine/Source/Programs/UnrealTraceServer/src/Recorder.cpp
--- a/Engine/Source/Programs/UnrealTraceServer/src/Recorder.cpp
+++ b/Engine/Source/Programs/UnrealTraceServer/src/Recorder.cpp
@@ -173,19 +173,6 @@ bool FRecorderRelay::ReadMagic()
 		return false;
 	}
 
-	// We can continue to support very old clients
-	if (Cursor[0] == 'E')
-	{
-		if (CreateTrace())
-		{
-			// Old clients have no metadata so we can go straight into the
-			// read-write loop. We've already got read data in Buffer.
-			OnIoComplete(OpBuffer0|OpFill, sizeof(MagicType) + sizeof(MetadataSizeType));
-			return true;
-		}
-		return false;
-	}
-
 	// Later clients have a metadata block (TRC2). There's loose support for the
 	// future too if need be (TRC[3-9]).
 	if (Cursor[0] < '2' || Cursor[0] > '9')
