Renamed makefiles' parent folder from build/ to scripts/
#rnx
(git:1b3e9c6 - Martin Ridgers - [REDACTED_PHONE]:51:27 +0100)

diff --git a/Engine/Source/Programs/UnrealTraceServer/makefile b/Engine/Source/Programs/UnrealTraceServer/makefile
--- a/Engine/Source/Programs/UnrealTraceServer/makefile
+++ b/Engine/Source/Programs/UnrealTraceServer/makefile
@@ -6,12 +6,12 @@ all :
 ifdef MAKEDIR: # gmake=false condition, nmake=two unused targets
 !ifdef MAKEDIR # gmake=skipped,			nmake=true condition
 
-!include build/nmake.makefile
+!include scripts/nmake.makefile
 
 !else
 else
 
-include build/gmake.makefile
+include scripts/gmake.makefile
 
 endif 
 !endif :
