Flatten makefile tree
#rnx
(git:a3ad22a - Martin Ridgers - [REDACTED_PHONE]:19:52 +0100)

diff --git a/Engine/Source/Programs/UnrealTraceServer/makefile b/Engine/Source/Programs/UnrealTraceServer/makefile
--- a/Engine/Source/Programs/UnrealTraceServer/makefile
+++ b/Engine/Source/Programs/UnrealTraceServer/makefile
@@ -6,12 +6,12 @@ all :
 ifdef MAKEDIR: # gmake=false condition, nmake=two unused targets
 !ifdef MAKEDIR # gmake=skipped,			nmake=true condition
 
-!include scripts/nmake.makefile
+!include nmake.makefile
 
 !else
 else
 
-include scripts/gmake.makefile
+include gmake.makefile
 
 endif 
 !endif :
