Ignore intermediate files

diff --git a/Engine/Source/Programs/UnrealTraceServer/.gitignore b/Engine/Source/Programs/UnrealTraceServer/.gitignore
--- a/Engine/Source/Programs/UnrealTraceServer/.gitignore
+++ b/Engine/Source/Programs/UnrealTraceServer/.gitignore
@@ -1,2 +1,4 @@
 _build
+_intermediate
 tags
+*.user
