Expose some std functionality in aliases

diff --git a/Engine/Source/Programs/UnrealTraceServer/src/Foundation.h b/Engine/Source/Programs/UnrealTraceServer/src/Foundation.h
--- a/Engine/Source/Programs/UnrealTraceServer/src/Foundation.h
+++ b/Engine/Source/Programs/UnrealTraceServer/src/Foundation.h
@@ -26,8 +26,47 @@ struct TArray
 	void		SetNumUninitialized(int Num)		{ return Super::resize(Num); }
 	void		Append(const void* Data, int Num)	{ Super::insert(Super::end(), (const Type*)Data, ((const Type*)Data) + Num); }
 	Type&		Last()								{ return Super::back(); }
+	template<typename Lambda>
+	uint32		RemoveIf(Lambda&& Pred);
+	template<typename Lambda> 
+	bool		FindOrAdd(const Type& Item, Lambda&& Pred);
 };
 
+template<typename Type>
+template<typename Lambda>
+inline uint32 TArray<Type>::RemoveIf(Lambda&& Pred) 
+{
+	uint32 Removed = 0;
+	for (auto it = begin(); it != end();) 
+	{
+		if (Pred(*it))
+		{
+			Super::erase(it); 
+			++Removed;
+		}
+		else
+		{
+			++it;
+		}
+	}
+	return Removed;
+}
+
+template<typename Type>
+template<typename Lambda>
+inline bool TArray<Type>::FindOrAdd(const Type& Item, Lambda&& Compare)
+{
+	for (auto it = begin(); it != end();)
+	{
+		if (Compare(*it) == 0)
+		{
+			return false;
+		}
+	}
+	Super::push_back(Item);
+	return true;
+}
+
 ////////////////////////////////////////////////////////////////////////////////
 struct FStringView
 	: protected std::basic_string_view<char>
@@ -39,6 +78,8 @@ struct FStringView
 	const char* GetData() const			{ return Super::data(); }
 	uint32 Len() const					{ return uint32(Super::size()); }
 	int Compare(const char* Rhs) const	{ return Super::compare(Rhs); }
+	bool StartsWith(const char* Rhs) const { return Super::starts_with(Rhs); }
+	void RemovePrefix(uint32 n)			{ Super::remove_prefix(n); }
 };
 
 ////////////////////////////////////////////////////////////////////////////////
@@ -51,6 +92,7 @@ struct FString
 
 					FString(const std::string& Rhs) : Super(Rhs)	{}
 					FString(std::string&& Rhs) : Super(Rhs)			{}
+					FString(FStringView Rhs) : Super(Rhs.GetData(), Rhs.Len())			{}
 	const char*		operator * () const								{ return Super::c_str(); }
 	void			operator += (const FStringView& Rhs)			{ Super::operator += ((FStringView::Super&)(Rhs)); }
 };
