Add functionality to determine if a connection is local.

diff --git a/Engine/Source/Programs/UnrealTraceServer/src/AsioSocket.cpp b/Engine/Source/Programs/UnrealTraceServer/src/AsioSocket.cpp
--- a/Engine/Source/Programs/UnrealTraceServer/src/AsioSocket.cpp
+++ b/Engine/Source/Programs/UnrealTraceServer/src/AsioSocket.cpp
@@ -77,6 +77,12 @@ void FAsioSocket::Close()
 	Socket.close();
 }
 
+////////////////////////////////////////////////////////////////////////////////
+bool FAsioSocket::IsLocalConnection() const
+{
+	return Socket.local_endpoint().address() == Socket.remote_endpoint().address();
+}
+
 ////////////////////////////////////////////////////////////////////////////////
 bool FAsioSocket::Read(void* Dest, uint32 Size, FAsioIoSink* Sink, uint32 Id)
 {
diff --git a/Engine/Source/Programs/UnrealTraceServer/src/AsioSocket.h b/Engine/Source/Programs/UnrealTraceServer/src/AsioSocket.h
--- a/Engine/Source/Programs/UnrealTraceServer/src/AsioSocket.h
+++ b/Engine/Source/Programs/UnrealTraceServer/src/AsioSocket.h
@@ -17,6 +17,7 @@ public:
 	virtual bool			IsOpen() const override;
 	virtual void			Close() override;
 	virtual bool			HasDataAvailable() const override;
+	bool					IsLocalConnection() const;
 	uint32					GetRemoteAddress() const;
 	uint32					GetRemotePort() const;
 	uint32					GetLocalPort() const;
