Persistent setting for Unreal Trace Server

* Adds persistent settings for Unreal Trace Server, stored in user settings directory.
* End point for modying configuration. Endpoint is only accessible from the local machine.
* Functionality to detect changes, update components and relay to clients using settings serial number.

diff --git a/Engine/Source/Programs/UnrealTraceServer/src/CborPayload.h b/Engine/Source/Programs/UnrealTraceServer/src/CborPayload.h
--- a/Engine/Source/Programs/UnrealTraceServer/src/CborPayload.h
+++ b/Engine/Source/Programs/UnrealTraceServer/src/CborPayload.h
@@ -21,6 +21,7 @@ enum class EStatusCode
 	Unknown				= 0,
 	Success				= 200,
 	BadRequest			= 400,
+	Forbidden			= 403,
 	NotFound			= 404,
 	InternalError		= 500,
 };
@@ -36,6 +37,7 @@ public:
 	template <int N>		TPayloadBuilder(const char (&Path)[N]);
 	template <int N> void	AddInteger(const char (&Name)[N], int64 Value);
 	template <int N> void	AddString(const char (&Name)[N], const char* Value, int32 Length=-1);
+	template <int N> void	AddStringArray(const char(&Name)[N], const TArray<FString> Values);
 	FPayload				Done();
 
 private:
@@ -70,6 +72,20 @@ inline void TPayloadBuilder<Size>::AddInteger(const char (&Name)[N], int64 Value
 	CborWriter.WriteInteger(Value);
 }
 
+////////////////////////////////////////////////////////////////////////////////
+template<int Size>
+template<int N>
+inline void TPayloadBuilder<Size>::AddStringArray(const char(&Name)[N], const TArray<FString> Values)
+{
+	CborWriter.WriteString(Name, N - 1);
+	CborWriter.OpenArray();
+	for (const auto& String : Values)
+	{
+		CborWriter.WriteString(*String);
+	}
+	CborWriter.Close();
+}
+
 ////////////////////////////////////////////////////////////////////////////////
 template <int Size>
 template <int N>
diff --git a/Engine/Source/Programs/UnrealTraceServer/src/Main.cpp b/Engine/Source/Programs/UnrealTraceServer/src/Main.cpp
--- a/Engine/Source/Programs/UnrealTraceServer/src/Main.cpp
+++ b/Engine/Source/Programs/UnrealTraceServer/src/Main.cpp
@@ -1,6 +1,7 @@
 // Copyright Epic Games, Inc. All Rights Reserved.
 
 #include "Pch.h"
+#include "StoreSettings.h"
 #include "StoreService.h"
 #include "Utils.h"
 #include "Version.h"
@@ -297,7 +298,7 @@ struct FStoreOptions
 };
 
 ////////////////////////////////////////////////////////////////////////////////
-static FStoreOptions ParseOptions(int ArgC, char** ArgV)
+static void ParseOptions(int ArgC, char** ArgV, FStoreSettings* Settings)
 {
 	cxxopts::Options Options("UnrealTraceServer", "Unreal Trace Server");
 	Options.add_options()
@@ -308,35 +309,22 @@ static FStoreOptions ParseOptions(int ArgC, char** ArgV)
 
 	auto Parsed = Options.parse(ArgC, ArgV);
 
-	FStoreOptions Ret;
 	if (int Value = Parsed["port"].as<int>())
 	{
-		Ret.Port = Value;
+		Settings->StorePort = Value;
 	}
 
 	if (int Value = Parsed["recport"].as<int>())
 	{
-		Ret.RecorderPort = Value;
+		Settings->RecorderPort = Value;
 	}
 
 	std::string Value = Parsed["storedir"].as<std::string>();
 	if (!Value.empty())
 	{
-		std::error_code ErrorCode;
-		Ret.Dir = fs::absolute(Value, ErrorCode);
+		Settings->StoreDir = Value;
 	}
 
-	return Ret;
-}
-
-////////////////////////////////////////////////////////////////////////////////
-static FStoreService* StartStore(const FStoreOptions& Options)
-{
-	FStoreService::FDesc Desc;
-	Desc.StoreDir = Options.Dir;
-	Desc.StorePort = Options.Port;
-	Desc.RecorderPort = Options.RecorderPort;
-	return FStoreService::Create(Desc);
 }
 
 
@@ -887,20 +875,32 @@ static int MainDaemon(int ArgC, char** ArgV)
 	}
 
 	// Fire up the store
-	TS_LOG("Starting the store");
-	FStoreService* StoreService;
-	{
-		FStoreOptions Options = ParseOptions(ArgC, ArgV);
-
-		if (Options.Dir.empty())
-		{
-			FPath StoreDir;
-			GetUnrealTraceHome(StoreDir);
-			StoreDir /= "Store";
-			Options.Dir = StoreDir;
-		}
-
-		StoreService = StartStore(Options);
+	FStoreSettings* Settings = new FStoreSettings();
+	FStoreService* StoreService = nullptr;
+	{
+		 TS_LOG("Starting the store");
+		 FPath HomeDir;
+		 GetUnrealTraceHome(HomeDir);
+
+		 // Read settings from configuration file
+		 Settings->ReadFromSettings(HomeDir);
+
+		 // Override with command line arguments.
+		 ParseOptions(ArgC, ArgV, Settings);
+
+		 TS_LOG("Creating store with settings:");
+		 TS_LOG(" - Directory: '%s'", Settings->StoreDir.string().c_str());
+		 TS_LOG(" - Store port: %u", Settings->StorePort);
+		 TS_LOG(" - Recorder port: %u", Settings->RecorderPort);
+		 if (Settings->AdditionalWatchDirs.Num())
+		 {
+			 for (const auto& Dir : Settings->AdditionalWatchDirs)
+			 {
+				 TS_LOG(" - Additional watch directory: '%s'", Dir.string().c_str());
+			 }
+		 }
+
+		 StoreService = FStoreService::Create(Settings);
 	}
 
 	// Let every one know we've started.
diff --git a/Engine/Source/Programs/UnrealTraceServer/src/Store.cpp b/Engine/Source/Programs/UnrealTraceServer/src/Store.cpp
--- a/Engine/Source/Programs/UnrealTraceServer/src/Store.cpp
+++ b/Engine/Source/Programs/UnrealTraceServer/src/Store.cpp
@@ -2,6 +2,7 @@
 
 #include "Pch.h"
 #include "Store.h"
+#include "StoreSettings.h"
 
 #if TS_USING(TS_PLATFORM_LINUX)
 #   include <sys/inotify.h>
@@ -25,7 +26,7 @@ static int64 FsToUnrealEpochBiasSeconds	= uint64(double(FsEpochYear - UnrealEpoc
 
 ////////////////////////////////////////////////////////////////////////////////
 #if TS_USING(TS_PLATFORM_WINDOWS)
-class FStore::FDirWatcher
+class FStore::FMount::FDirWatcher
 	: public asio::windows::object_handle
 {
 public:
@@ -34,7 +35,7 @@ public:
 
 ////////////////////////////////////////////////////////////////////////////////
 #elif TS_USING(TS_PLATFORM_LINUX)
-class FStore::FDirWatcher
+class FStore::FMount::FDirWatcher
 	: public asio::posix::stream_descriptor
 {
 	typedef std::function<void(asio::error_code error)> HandlerType;
@@ -48,7 +49,7 @@ public:
 
 ////////////////////////////////////////////////////////////////////////////////
 #elif TS_USING(TS_PLATFORM_MAC)
-class FStore::FDirWatcher
+class FStore::FMount::FDirWatcher
 {
 	typedef std::function<void(asio::error_code error)> HandlerType;
 public:
@@ -95,7 +96,7 @@ static void MacCallback(ConstFSEventStreamRef StreamRef,
 					const FSEventStreamEventFlags EventFlags[],
 					const FSEventStreamEventId EventIDs[])
 {
-	FStore::FDirWatcher* DirWatcherPtr = (FStore::FDirWatcher*)InDirWatcherPtr;
+	FStore::FMount::FDirWatcher* DirWatcherPtr = (FStore::FMount::FDirWatcher*)InDirWatcherPtr;
 	check(DirWatcherPtr);
 	check(DirWatcherPtr->EventStream == StreamRef);
 
@@ -103,7 +104,7 @@ static void MacCallback(ConstFSEventStreamRef StreamRef,
 }
 
 ////////////////////////////////////////////////////////////////////////////////
-void FStore::FDirWatcher::async_wait(HandlerType InHandler)
+void FStore::FMount::FDirWatcher::async_wait(HandlerType InHandler)
 {
 	if (bIsRunning)
 	{
@@ -186,8 +187,9 @@ uint32 FStore::FTrace::GetId() const
 ////////////////////////////////////////////////////////////////////////////////
 uint64 FStore::FTrace::GetSize() const
 {
-	std::error_code Ec;
-	return std::filesystem::file_size(Path, Ec);
+	std::error_code Error;
+	const uint64 Size = std::filesystem::file_size(Path, Error);
+	return Error ? 0 : Size;
 }
 
 ////////////////////////////////////////////////////////////////////////////////
@@ -196,15 +198,58 @@ uint64 FStore::FTrace::GetTimestamp() const
 	return Timestamp;
 }
 
-
+////////////////////////////////////////////////////////////////////////////////
+FStore::FMount* FStore::FMount::Create(FStore* InParent, asio::io_context& InIoContext, const FPath& InDir)
+{
+	std::error_code ErrorCodeAbs, ErrorCodeCreate;
+	// We would like to check if the path is absolute here, but Microsoft
+	// has a very specific interpretation of what constitutes absolute path
+	// which doesn't correspond to what UE's file utilities does. Hence paths
+	// from Insights will not be correctly formatted.
+	//const FPath AbsoluteDir = InDir.is_absolute() ? InDir : fs::absolute(InDir, ErrorCodeAbs);
+	const FPath AbsoluteDir = InDir;
+	fs::create_directories(AbsoluteDir, ErrorCodeCreate);
+	if (ErrorCodeAbs || ErrorCodeCreate)
+	{
+		return nullptr;
+	}
+	return new FMount(InParent, InIoContext, AbsoluteDir);
+}
 
 ////////////////////////////////////////////////////////////////////////////////
-FStore::FMount::FMount(const FPath& InDir)
+FStore::FMount::FMount(FStore* InParent, asio::io_context& InIoContext, const fs::path& InDir)
 : Id(QuickStoreHash(InDir.c_str()))
+, Dir(InDir)
+, Parent(InParent)
+, IoContext(InIoContext)
 {
-	std::error_code ErrorCode;
-	Dir = fs::absolute(InDir, ErrorCode);
-	fs::create_directories(Dir, ErrorCode);
+
+#if TS_USING(TS_PLATFORM_WINDOWS)
+	std::wstring StoreDirW = Dir;
+	HANDLE DirWatchHandle = FindFirstChangeNotificationW(StoreDirW.c_str(), false, FILE_NOTIFY_CHANGE_FILE_NAME|FILE_NOTIFY_CHANGE_DIR_NAME);
+	if (DirWatchHandle == INVALID_HANDLE_VALUE)
+	{
+		DirWatchHandle = 0;
+	}
+	DirWatcher = new FDirWatcher(IoContext, DirWatchHandle);
+#elif TS_USING(TS_PLATFORM_LINUX)
+	int inotfd = inotify_init();
+	int watch_desc = inotify_add_watch(inotfd, StoreDir.c_str(), IN_CREATE | IN_DELETE);
+	DirWatcher = new FDirWatcher(IoContext, inotfd);
+#elif PLATFORM_MAC
+	DirWatcher = new FDirWatcher(*StoreDir);
+#endif
+
+	WatchDir();
+}
+
+////////////////////////////////////////////////////////////////////////////////
+FStore::FMount::~FMount()
+{
+	if (DirWatcher != nullptr)
+	{
+		delete DirWatcher;
+	}
 }
 
 ////////////////////////////////////////////////////////////////////////////////
@@ -213,6 +258,22 @@ FString FStore::FMount::GetDir() const
 	return fs::ToFString(Dir);
 }
 
+////////////////////////////////////////////////////////////////////////////////
+const FPath& FStore::FMount::GetPath() const
+{
+	return Dir;
+}
+
+////////////////////////////////////////////////////////////////////////////////
+void FStore::FMount::Close()
+{
+	if (DirWatcher != nullptr)
+	{
+		DirWatcher->cancel();
+		DirWatcher->close();
+	}
+}
+
 ////////////////////////////////////////////////////////////////////////////////
 uint32 FStore::FMount::GetId() const
 {
@@ -267,10 +328,15 @@ FStore::FTrace* FStore::FMount::AddTrace(const FPath& Path)
 }
 
 ////////////////////////////////////////////////////////////////////////////////
-uint32 FStore::FMount::Refresh()
+void FStore::FMount::ClearTraces()
 {
 	Traces.Empty();
+}
 
+////////////////////////////////////////////////////////////////////////////////
+uint32 FStore::FMount::Refresh()
+{
+	ClearTraces();
 	uint32 ChangeSerial = 0;
 	for (auto& DirItem : std::filesystem::directory_iterator(Dir))
 	{
@@ -296,67 +362,53 @@ uint32 FStore::FMount::Refresh()
 }
 
 
-
-////////////////////////////////////////////////////////////////////////////////
-FStore::FStore(asio::io_context& InIoContext, const FPath& InStoreDir)
-: IoContext(InIoContext)
+void FStore::SetupMounts()
 {
-	FPath StoreDir = InStoreDir / "001";
+	FPath StoreDir = Settings->StoreDir / "001";
 	AddMount(StoreDir);
 
-	Refresh();
-
-#if TS_USING(TS_PLATFORM_WINDOWS)
-	std::wstring StoreDirW = StoreDir;
-	HANDLE DirWatchHandle = FindFirstChangeNotificationW(StoreDirW.c_str(), false,
-		FILE_NOTIFY_CHANGE_FILE_NAME|FILE_NOTIFY_CHANGE_DIR_NAME);
-	if (DirWatchHandle == INVALID_HANDLE_VALUE)
+	for (const FPath& Path : Settings->AdditionalWatchDirs)
 	{
-		DirWatchHandle = 0;
+		AddMount(Path);
 	}
-	DirWatcher = new FDirWatcher(IoContext, DirWatchHandle);
-#elif TS_USING(TS_PLATFORM_LINUX)
-	int inotfd = inotify_init();
-	int watch_desc = inotify_add_watch(inotfd, StoreDir.c_str(), IN_CREATE|IN_DELETE);
-	DirWatcher = new FDirWatcher(IoContext, inotfd);
-#elif PLATFORM_MAC
-	DirWatcher = new FDirWatcher(*StoreDir);
-#endif
+}
 
-	WatchDir();
+////////////////////////////////////////////////////////////////////////////////
+FStore::FStore(asio::io_context& InIoContext, const FStoreSettings* InSettings)
+: IoContext(InIoContext)
+, Settings(InSettings)
+{
+	SetupMounts();
+	Refresh();
 }
 
 ////////////////////////////////////////////////////////////////////////////////
 FStore::~FStore()
 {
-	if (DirWatcher != nullptr)
+	for (FMount* Mount : Mounts)
 	{
-		check(!DirWatcher->is_open());
-		delete DirWatcher;
+		delete Mount;
 	}
 }
 
 ////////////////////////////////////////////////////////////////////////////////
 void FStore::Close()
 {
-	if (DirWatcher != nullptr)
-	{
-		DirWatcher->cancel();
-		DirWatcher->close();
-	}
-
 	for (FMount* Mount : Mounts)
 	{
-		delete Mount;
+		Mount->Close();
 	}
 }
 
 ////////////////////////////////////////////////////////////////////////////////
 bool FStore::AddMount(const FPath& Dir)
 {
-	FMount* Mount = new FMount(Dir);
-	Mounts.Add(Mount);
-	return true;
+	FMount* Mount = FMount::Create(this, IoContext, Dir);
+	if (Mount)
+	{
+		Mounts.Add(Mount);
+	}
+	return Mount != nullptr;
 }
 
 ////////////////////////////////////////////////////////////////////////////////
@@ -377,38 +429,7 @@ bool FStore::RemoveMount(uint32 Id)
 }
 
 ////////////////////////////////////////////////////////////////////////////////
-const FStore::FMount* FStore::GetMount(uint32 Id) const
-{
-	for (FMount* Mount : Mounts)
-	{
-		if (Mount->GetId() == Id)
-		{
-			return Mount;
-		}
-	}
-
-	return nullptr;
-}
-
-////////////////////////////////////////////////////////////////////////////////
-uint32 FStore::GetMountCount() const
-{
-	return uint32(Mounts.Num());
-}
-
-////////////////////////////////////////////////////////////////////////////////
-const FStore::FMount* FStore::GetMountInfo(uint32 Index) const
-{
-	if (Index >= uint32(Mounts.Num()))
-	{
-		return nullptr;
-	}
-
-	return Mounts[Index];
-}
-
-////////////////////////////////////////////////////////////////////////////////
-void FStore::WatchDir()
+void FStore::FMount::WatchDir()
 {
 	if (DirWatcher == nullptr)
 	{
@@ -432,7 +453,7 @@ void FStore::WatchDir()
 		{
 			delete DelayTimer;
 
-			Refresh();
+			Parent->Refresh();
 
 			FindNextChangeNotification(DirWatcher->native_handle());
 			WatchDir();
@@ -559,6 +580,19 @@ FAsioReadable* FStore::OpenTrace(uint32 Id)
 	return FAsioFile::ReadFile(IoContext, Trace->GetPath());
 }
 
+////////////////////////////////////////////////////////////////////////////////
+void FStore::OnSettingsChanged()
+{
+	// Remove all existing mounts
+	for (const FMount* Mount : Mounts)
+	{
+		delete Mount;
+	}
+	Mounts.Empty();
+	SetupMounts();
+	Refresh();
+}
+
 ////////////////////////////////////////////////////////////////////////////////
 void FStore::Refresh()
 {
diff --git a/Engine/Source/Programs/UnrealTraceServer/src/Store.h b/Engine/Source/Programs/UnrealTraceServer/src/Store.h
--- a/Engine/Source/Programs/UnrealTraceServer/src/Store.h
+++ b/Engine/Source/Programs/UnrealTraceServer/src/Store.h
@@ -9,6 +9,8 @@
 
 class FAsioReadable;
 class FAsioWriteable;
+class FDirWatcher;
+class FStoreSettings;
 
 ////////////////////////////////////////////////////////////////////////////////
 enum class EStoreVersion
@@ -46,31 +48,37 @@ public:
 	class FMount
 	{
 	public:
-						FMount(const FPath& InDir);
+		static FMount*	Create(FStore* InParent, asio::io_context& InIoContext, const FPath& InDir);
+						
+		void			Close();
 		uint32			GetId() const;
 		FString			GetDir() const;
+		const FPath&	GetPath() const;
 		uint32			GetTraceCount() const;
 		const FTrace*	GetTraceInfo(uint32 Index) const;
 
 	private:
+		class			FDirWatcher;
 		friend			FStore;
+						FMount(FStore* InParent, asio::io_context& InIoContext, const FPath& InDir);
+						~FMount();
 		FTrace*			GetTrace(uint32 Id) const;
 		FTrace*			AddTrace(const FPath& Path);
+		void			ClearTraces();
 		uint32			Refresh();
+		void			WatchDir();
 		FPath			Dir;
 		TArray<FTrace*>	Traces;
 		uint32			Id;
+		FDirWatcher*	DirWatcher = nullptr;
+		FStore*			Parent = nullptr;
+		asio::io_context& IoContext;
 	};
 
-						FStore(asio::io_context& IoContext, const FPath& InStoreDir);
+	void SetupMounts();
+						FStore(asio::io_context& IoContext, const FStoreSettings* InSettings);
 						~FStore();
 	void				Close();
-	bool				AddMount(const FPath& Dir);
-	bool				RemoveMount(uint32 Id);
-	const FMount*		GetMount(uint32 Id) const;
-	uint32				GetMountCount() const;
-	const FMount*		GetMountInfo(uint32 Index) const;
-
 	FString				GetStoreDir() const;
 	uint32				GetChangeSerial() const;
 	uint32				GetTraceCount() const;
@@ -78,16 +86,18 @@ public:
 	bool				HasTrace(uint32 Id) const;
 	FNewTrace			CreateTrace();
 	FAsioReadable*		OpenTrace(uint32 Id);
-	class				FDirWatcher;
+	void				OnSettingsChanged();
 
 private:
+	bool				AddMount(const FPath& Dir);
+	bool				RemoveMount(uint32 Id);
 	FTrace*				GetTrace(uint32 Id, FMount** OutMount=nullptr) const;
 	void				Refresh();
-	void				WatchDir();
+	
 	asio::io_context&	IoContext;
 	TArray<FMount*>		Mounts;
 	uint32				ChangeSerial;
-	FDirWatcher*		DirWatcher = nullptr;
+	const FStoreSettings* Settings;
 };
 
 /* vim: set noexpandtab : */
diff --git a/Engine/Source/Programs/UnrealTraceServer/src/StoreCborServer.cpp b/Engine/Source/Programs/UnrealTraceServer/src/StoreCborServer.cpp
--- a/Engine/Source/Programs/UnrealTraceServer/src/StoreCborServer.cpp
+++ b/Engine/Source/Programs/UnrealTraceServer/src/StoreCborServer.cpp
@@ -4,10 +4,12 @@
 #include "AsioIoable.h"
 #include "AsioSocket.h"
 #include "CborPayload.h"
+#include "StoreService.h"
 #include "Recorder.h"
 #include "Store.h"
 #include "StoreCborServer.h"
 #include "TraceRelay.h"
+#include "StoreSettings.h"
 
 ////////////////////////////////////////////////////////////////////////////////
 class FStoreCborPeer
@@ -27,6 +29,7 @@ protected:
 	void					OnTraceCount();
 	void					OnTraceInfo();
 	void					OnTraceRead();
+	void					OnSettingsWrite();
 	void					SendError(EStatusCode StatusCode);
 	void					SendResponse(const FPayload& Payload);
 	virtual void			OnIoComplete(uint32 Id, int32 Size) override;
@@ -130,10 +133,18 @@ void FStoreCborPeer::OnSessionInfo()
 ////////////////////////////////////////////////////////////////////////////////
 void FStoreCborPeer::OnStatus()
 {
+	const FStoreSettings* Settings = Parent.GetSettings();
 	TPayloadBuilder<> Builder(EStatusCode::Success);
-	Builder.AddInteger("recorder_port", Recorder.GetPort());
+	Builder.AddInteger("recorder_port", Settings->RecorderPort);
 	Builder.AddInteger("change_serial", Store.GetChangeSerial());
-	Builder.AddString("store_dir", *Store.GetStoreDir());
+	Builder.AddInteger("settings_serial", Settings->GetChangeSerial());
+	Builder.AddString("store_dir", Settings->StoreDir.string().c_str());
+	TArray<FString> AdditionalWatchDirs;
+	for (const fs::path& Path : Settings->AdditionalWatchDirs)
+	{
+		AdditionalWatchDirs.Add(Path.string());
+	}
+	Builder.AddStringArray("watch_dirs", AdditionalWatchDirs);
 	SendResponse(Builder.Done());
 }
 
@@ -178,12 +189,14 @@ void FStoreCborPeer::OnTraceInfo()
 	}
 
 	auto Name = Trace->GetName();
+	auto PathString = Trace->GetPath().string();
 
 	TPayloadBuilder<> Builder(EStatusCode::Success);
 	Builder.AddInteger("id", Trace->GetId());
 	Builder.AddInteger("size", Trace->GetSize());
 	Builder.AddInteger("timestamp", Trace->GetTimestamp());
 	Builder.AddString("name", *Name);
+	Builder.AddString("uri", PathString.c_str());
 	SendResponse(Builder.Done());
 }
 
@@ -207,6 +220,29 @@ void FStoreCborPeer::OnTraceRead()
 	SendResponse(Builder.Done());
 }
 
+////////////////////////////////////////////////////////////////////////////////
+void FStoreCborPeer::OnSettingsWrite()
+{
+	// Only allow local connections to change settings
+	if (!Socket.IsLocalConnection())
+	{
+		return SendError(EStatusCode::Forbidden);
+	}
+	// Copy the current settings and apply requested changes. This is possible
+	// to do without synchronization since asio currently only runs on a single
+	// thread, so no other operations should be running at the same time.
+	FStoreSettings* Settings = Parent.GetSettings();
+	Settings->ApplySettingsFromCbor(Response.GetData(), Response.GetSize());
+	// Save settings to disk
+	Settings->WriteToSettingsFile();
+
+	TPayloadBuilder<> Builder(EStatusCode::Success);
+	SendResponse(Builder.Done());
+
+	// Let all the service components know settings have changed
+	Parent.OnSettingsChanged();
+}
+
 ////////////////////////////////////////////////////////////////////////////////
 void FStoreCborPeer::OnPayload()
 {
@@ -228,6 +264,7 @@ void FStoreCborPeer::OnPayload()
 		{ QuickStoreHash("v1/trace/count"),		&FStoreCborPeer::OnTraceCount },
 		{ QuickStoreHash("v1/trace/info"),		&FStoreCborPeer::OnTraceInfo },
 		{ QuickStoreHash("v1/trace/read"),		&FStoreCborPeer::OnTraceRead },
+		{ QuickStoreHash("v1/settings/write"),	&FStoreCborPeer::OnSettingsWrite },
 	};
 
 	uint32 PathHash = QuickStoreHash(Path);
@@ -300,15 +337,16 @@ void FStoreCborPeer::OnIoComplete(uint32 Id, int32 Size)
 ////////////////////////////////////////////////////////////////////////////////
 FStoreCborServer::FStoreCborServer(
 	asio::io_context& IoContext,
-	int32 Port,
+	FStoreSettings* InSettings,
 	FStore& InStore,
 	FRecorder& InRecorder)
 : FAsioTcpServer(IoContext)
 , FAsioTickable(IoContext)
 , Store(InStore)
 , Recorder(InRecorder)
+, Settings(InSettings)
 {
-	if (!StartServer(Port))
+	if (!StartServer(Settings->StorePort))
 	{
 		StartServer();
 	}
@@ -437,4 +475,10 @@ FTraceRelay* FStoreCborServer::RelayTrace(uint32 Id)
 	return Relay;
 }
 
+////////////////////////////////////////////////////////////////////////////////
+void FStoreCborServer::OnSettingsChanged()
+{
+	Store.OnSettingsChanged();
+}
+
 /* vim: set noexpandtab : */
diff --git a/Engine/Source/Programs/UnrealTraceServer/src/StoreCborServer.h b/Engine/Source/Programs/UnrealTraceServer/src/StoreCborServer.h
--- a/Engine/Source/Programs/UnrealTraceServer/src/StoreCborServer.h
+++ b/Engine/Source/Programs/UnrealTraceServer/src/StoreCborServer.h
@@ -11,27 +11,34 @@ class FRecorder;
 class FStore;
 class FStoreCborPeer;
 class FTraceRelay;
+class FStoreSettings;
 
 ////////////////////////////////////////////////////////////////////////////////
 class FStoreCborServer
 	: public FAsioTcpServer
 	, public FAsioTickable
 {
+	friend class FStoreCborPeer;
+
 public:
-							FStoreCborServer(asio::io_context& IoContext, int32 Port, FStore& InStore, FRecorder& InRecorder);
+							FStoreCborServer(asio::io_context& IoContext, FStoreSettings* InSettings, FStore& InStore, FRecorder& InRecorder);
 							~FStoreCborServer();
 	void					Close();
 	FStore&					GetStore() const;
 	FRecorder&				GetRecorder() const;
 	FTraceRelay*			RelayTrace(uint32 Id);
+	void					OnSettingsChanged();
 
 private:
 	virtual bool			OnAccept(asio::ip::tcp::socket& Socket) override;
 	virtual void			OnTick() override;
+	FStoreSettings*			GetSettings() const { return Settings; }
+
 	TArray<FStoreCborPeer*>	Peers;
 	TArray<FTraceRelay*>	Relays;
 	FStore&					Store;
 	FRecorder&				Recorder;
+	FStoreSettings*			Settings;
 };
 
 /* vim: set noexpandtab : */
diff --git a/Engine/Source/Programs/UnrealTraceServer/src/StoreService.cpp b/Engine/Source/Programs/UnrealTraceServer/src/StoreService.cpp
--- a/Engine/Source/Programs/UnrealTraceServer/src/StoreService.cpp
+++ b/Engine/Source/Programs/UnrealTraceServer/src/StoreService.cpp
@@ -7,12 +7,13 @@
 #include "Store.h"
 #include "StoreCborServer.h"
 #include "StoreService.h"
+#include "StoreSettings.h"
 
 ////////////////////////////////////////////////////////////////////////////////
 struct FStoreServiceImpl
 {
 public:
-							FStoreServiceImpl(const FStoreService::FDesc& Desc);
+							FStoreServiceImpl(FStoreSettings* Settings);
 							~FStoreServiceImpl();
 	FAsioContext			Context;
 	FStore					Store;
@@ -21,12 +22,17 @@ public:
 };
 
 ////////////////////////////////////////////////////////////////////////////////
-FStoreServiceImpl::FStoreServiceImpl(const FStoreService::FDesc& Desc)
-: Context(Desc.ThreadCount)
-, Store(Context.Get(), Desc.StoreDir)
+FStoreServiceImpl::FStoreServiceImpl(FStoreSettings* Settings)
+: Context(Settings->ThreadCount)
+, Store(Context.Get(), Settings)
 , Recorder(Context.Get(), Store)
-, CborServer(Context.Get(), Desc.StorePort, Store, Recorder)
+, CborServer(Context.Get(), Settings, Store, Recorder)
 {
+	if (Settings->RecorderPort >= 0)
+	{
+		Recorder.StartServer(Settings->RecorderPort);
+	}
+	Context.Start();
 }
 
 ////////////////////////////////////////////////////////////////////////////////
@@ -43,28 +49,17 @@ FStoreServiceImpl::~FStoreServiceImpl()
 
 
 ////////////////////////////////////////////////////////////////////////////////
-FStoreService* FStoreService::Create(const FDesc& InDesc)
+FStoreService* FStoreService::Create(FStoreSettings* Settings)
 {
-	FDesc Desc = InDesc;
-
-	if (Desc.ThreadCount <= 0)
+	if (Settings->ThreadCount <= 0)
 	{
-		Desc.ThreadCount = std::thread::hardware_concurrency();
+		Settings->ThreadCount = std::thread::hardware_concurrency();
 	}
 
 	// TODO: not thread safe yet
-	Desc.ThreadCount = 1;
-
-	FStoreServiceImpl* Impl = new FStoreServiceImpl(Desc);
-	if (Desc.RecorderPort >= 0)
-	{
-		FRecorder& Recorder = Impl->Recorder;
-		Recorder.StartServer(Desc.RecorderPort);
-	}
-
-	Impl->Context.Start();
+	Settings->ThreadCount = 1;
 
-	return (FStoreService*)Impl;
+	return (FStoreService*) new FStoreServiceImpl(Settings);
 }
 
 ////////////////////////////////////////////////////////////////////////////////
diff --git a/Engine/Source/Programs/UnrealTraceServer/src/StoreService.h b/Engine/Source/Programs/UnrealTraceServer/src/StoreService.h
--- a/Engine/Source/Programs/UnrealTraceServer/src/StoreService.h
+++ b/Engine/Source/Programs/UnrealTraceServer/src/StoreService.h
@@ -4,20 +4,14 @@
 
 #include "Foundation.h"
 
+class FStoreSettings;
+
 ////////////////////////////////////////////////////////////////////////////////
 class FStoreService
 {
 public:
-	struct FDesc
-	{
-		FPath				StoreDir;
-		int32				StorePort		= 0; // <=0:auto-assign
-		int32				RecorderPort	= 0; // 0:auto-assign, -1:off
-		int32				ThreadCount		= 0; // <=0:logical CPU count
-	};
-
 							~FStoreService() = default;
-	static FStoreService*	Create(const FDesc& Desc);
+	static FStoreService*	Create(FStoreSettings* Desc);
 	void					operator delete (void* Addr);
 	uint32					GetPort() const;
 	uint32					GetRecorderPort() const;
