Fixes to settings implementation

* Check for uniqeness when adding watch directories.
* Close file handle after done reading

diff --git a/Engine/Source/Programs/UnrealTraceServer/src/Foundation.h b/Engine/Source/Programs/UnrealTraceServer/src/Foundation.h
--- a/Engine/Source/Programs/UnrealTraceServer/src/Foundation.h
+++ b/Engine/Source/Programs/UnrealTraceServer/src/Foundation.h
@@ -56,9 +56,9 @@ template<typename Type>
 template<typename Lambda>
 inline bool TArray<Type>::FindOrAdd(const Type& Item, Lambda&& Compare)
 {
-	for (auto it = begin(); it != end();)
+	for (auto It = begin(); It != end(); ++It)
 	{
-		if (Compare(*it) == 0)
+		if (Compare(*It, Item) == 0)
 		{
 			return false;
 		}
diff --git a/Engine/Source/Programs/UnrealTraceServer/src/StoreSettings.cpp b/Engine/Source/Programs/UnrealTraceServer/src/StoreSettings.cpp
--- a/Engine/Source/Programs/UnrealTraceServer/src/StoreSettings.cpp
+++ b/Engine/Source/Programs/UnrealTraceServer/src/StoreSettings.cpp
@@ -123,6 +123,7 @@ public:
 		if (!Read)
 		{
 			TS_SETTINGS_LOG("Unable to read settings file.")
+			std::fclose(Handle);
 			return;
 		}
 
@@ -142,6 +143,8 @@ public:
 			BufferView.remove_prefix(std::min(BufferView.size(), Line.size() + 1));
 			Line = BufferView.substr(0, BufferView.find_first_of('\n'));
 		}
+
+		std::fclose(Handle);
 	}
 
 	void GetPath(const char* Ident, FPath& OutPath)
@@ -341,9 +344,17 @@ void FStoreSettings::ApplySettingsFromCbor(const uint8* Buffer, uint32 NumBytes)
 					else /*if (Value.Len() > 0)*/
 					{
 						FString ValueString(Value);
-						TS_SETTINGS_TRACE("Add -> %s '%s'\n", GAdditionalWatchDirsName, *ValueString)
 						
-						AdditionalWatchDirs.Add(fs::path(*ValueString));
+						if (AdditionalWatchDirs.FindOrAdd(fs::path(*ValueString), [](const FPath& Path, const FPath& InPath) { 
+							return Path.compare(InPath); 
+						}))
+						{
+							TS_SETTINGS_TRACE("Add -> %s '%s'\n", GAdditionalWatchDirsName, *ValueString)
+						}
+						else
+						{
+							TS_SETTINGS_TRACE("Did not add %s, existing directory with the same path was already added\n", *ValueString)
+						}
 					}
 				}
 			}
