Remove "001" folder when setting new store directories.

diff --git a/Engine/Source/Programs/UnrealTraceServer/src/Store.cpp b/Engine/Source/Programs/UnrealTraceServer/src/Store.cpp
--- a/Engine/Source/Programs/UnrealTraceServer/src/Store.cpp
+++ b/Engine/Source/Programs/UnrealTraceServer/src/Store.cpp
@@ -364,8 +364,7 @@ uint32 FStore::FMount::Refresh()
 
 void FStore::SetupMounts()
 {
-	FPath StoreDir = Settings->StoreDir / "001";
-	AddMount(StoreDir);
+	AddMount(Settings->StoreDir);
 
 	for (const FPath& Path : Settings->AdditionalWatchDirs)
 	{
diff --git a/Engine/Source/Programs/UnrealTraceServer/src/StoreSettings.cpp b/Engine/Source/Programs/UnrealTraceServer/src/StoreSettings.cpp
--- a/Engine/Source/Programs/UnrealTraceServer/src/StoreSettings.cpp
+++ b/Engine/Source/Programs/UnrealTraceServer/src/StoreSettings.cpp
@@ -264,7 +264,8 @@ void FStoreSettings::ReadFromSettings(const FPath& Path)
 	SettingsFile = FilePath.string();
 
 	// Set default store directory. Will be overwritten by settings
-	StoreDir = fs::absolute(Path / fs::path("Store"));
+	// Use the legacy '001' folder here to avoid changing existing stores.
+	StoreDir = fs::absolute(Path / fs::path("Store/001"));
 	
 	TS_SETTINGS_LOGF("Reading settings from '%s'\n", SettingsFile.string().c_str())
 
