Set locale for std methods.

Testing has reported that non-ascii characters in folder names creates folders with garbage names. This sets the global locale for all std types which help to parse stored folder names.

diff --git a/Engine/Source/Programs/UnrealTraceServer/src/Main.cpp b/Engine/Source/Programs/UnrealTraceServer/src/Main.cpp
--- a/Engine/Source/Programs/UnrealTraceServer/src/Main.cpp
+++ b/Engine/Source/Programs/UnrealTraceServer/src/Main.cpp
@@ -874,6 +874,9 @@ static int MainDaemon(int ArgC, char** ArgV)
 		InstanceInfo->Set();
 	}
 
+	// Set locale for std function application wide.
+	std::setlocale(LC_ALL, ".UTF-8");
+
 	// Fire up the store
 	FStoreSettings* Settings = new FStoreSettings();
 	FStoreService* StoreService = nullptr;
