Fixes related to watch folders.

* Avoid crashing if user deletes a mounted directory while running
* Avoid recreating watch folders when starting server. Don't create mounts for non-existing watch folders.

diff --git a/Engine/Source/Programs/UnrealTraceServer/src/Store.cpp b/Engine/Source/Programs/UnrealTraceServer/src/Store.cpp
--- a/Engine/Source/Programs/UnrealTraceServer/src/Store.cpp
+++ b/Engine/Source/Programs/UnrealTraceServer/src/Store.cpp
@@ -199,7 +199,7 @@ uint64 FStore::FTrace::GetTimestamp() const
 }
 
 ////////////////////////////////////////////////////////////////////////////////
-FStore::FMount* FStore::FMount::Create(FStore* InParent, asio::io_context& InIoContext, const FPath& InDir)
+FStore::FMount* FStore::FMount::Create(FStore* InParent, asio::io_context& InIoContext, const FPath& InDir, bool bCreate)
 {
 	std::error_code ErrorCodeAbs, ErrorCodeCreate;
 	// We would like to check if the path is absolute here, but Microsoft
@@ -208,10 +208,23 @@ FStore::FMount* FStore::FMount::Create(FStore* InParent, asio::io_context& InIoC
 	// from Insights will not be correctly formatted.
 	//const FPath AbsoluteDir = InDir.is_absolute() ? InDir : fs::absolute(InDir, ErrorCodeAbs);
 	const FPath AbsoluteDir = InDir;
-	fs::create_directories(AbsoluteDir, ErrorCodeCreate);
-	if (ErrorCodeAbs || ErrorCodeCreate)
+	if (bCreate)
 	{
-		return nullptr;
+		// Make sure the directory exists
+		fs::create_directories(AbsoluteDir, ErrorCodeCreate);
+		if (ErrorCodeAbs || ErrorCodeCreate)
+		{
+			return nullptr;
+		}
+	}
+	else
+	{
+		// If we are not allowed to create and if the directory does not exist
+		// there is no point adding a mount for it.
+		if (!fs::is_directory(AbsoluteDir))
+		{
+			return nullptr;
+		}
 	}
 	return new FMount(InParent, InIoContext, AbsoluteDir);
 }
@@ -338,7 +351,8 @@ uint32 FStore::FMount::Refresh()
 {
 	ClearTraces();
 	uint32 ChangeSerial = 0;
-	for (auto& DirItem : std::filesystem::directory_iterator(Dir))
+	std::error_code Ec;
+	for (auto& DirItem : std::filesystem::directory_iterator(Dir, Ec))
 	{
 		if (DirItem.is_directory())
 		{
@@ -364,11 +378,11 @@ uint32 FStore::FMount::Refresh()
 
 void FStore::SetupMounts()
 {
-	AddMount(Settings->StoreDir);
+	AddMount(Settings->StoreDir, true);
 
 	for (const FPath& Path : Settings->AdditionalWatchDirs)
 	{
-		AddMount(Path);
+		AddMount(Path, false);
 	}
 }
 
@@ -400,9 +414,9 @@ void FStore::Close()
 }
 
 ////////////////////////////////////////////////////////////////////////////////
-bool FStore::AddMount(const FPath& Dir)
+bool FStore::AddMount(const FPath& Dir, bool bCreate)
 {
-	FMount* Mount = FMount::Create(this, IoContext, Dir);
+	FMount* Mount = FMount::Create(this, IoContext, Dir, bCreate);
 	if (Mount)
 	{
 		Mounts.Add(Mount);
diff --git a/Engine/Source/Programs/UnrealTraceServer/src/Store.h b/Engine/Source/Programs/UnrealTraceServer/src/Store.h
--- a/Engine/Source/Programs/UnrealTraceServer/src/Store.h
+++ b/Engine/Source/Programs/UnrealTraceServer/src/Store.h
@@ -48,7 +48,7 @@ public:
 	class FMount
 	{
 	public:
-		static FMount*	Create(FStore* InParent, asio::io_context& InIoContext, const FPath& InDir);
+		static FMount* Create(FStore* InParent, asio::io_context& InIoContext, const FPath& InDir, bool bCreate);
 						
 		void			Close();
 		uint32			GetId() const;
@@ -89,7 +89,7 @@ public:
 	void				OnSettingsChanged();
 
 private:
-	bool				AddMount(const FPath& Dir);
+	bool				AddMount(const FPath& Dir, bool bCreate);
 	bool				RemoveMount(uint32 Id);
 	FTrace*				GetTrace(uint32 Id, FMount** OutMount=nullptr) const;
 	void				Refresh();
