New UTS version

* Fixed issues with non-ascii folder names.
* Do not crash if folder is deleted by user.
* Do not recreate folders that has been deleted.

diff --git a/Engine/Source/Programs/UnrealTraceServer/src/Version.h b/Engine/Source/Programs/UnrealTraceServer/src/Version.h
--- a/Engine/Source/Programs/UnrealTraceServer/src/Version.h
+++ b/Engine/Source/Programs/UnrealTraceServer/src/Version.h
@@ -7,7 +7,7 @@
 #endif
 
 #if !defined(TS_VERSION_MINOR)
-#	define TS_VERSION_MINOR		14
+#	define TS_VERSION_MINOR		15
 #endif
 
 #define TS_ICON_ID				493
