Remove StartsWith which is not present in C++17

diff --git a/Engine/Source/Programs/UnrealTraceServer/src/Foundation.h b/Engine/Source/Programs/UnrealTraceServer/src/Foundation.h
--- a/Engine/Source/Programs/UnrealTraceServer/src/Foundation.h
+++ b/Engine/Source/Programs/UnrealTraceServer/src/Foundation.h
@@ -77,8 +77,8 @@ struct FStringView
 
 	const char* GetData() const			{ return Super::data(); }
 	uint32 Len() const					{ return uint32(Super::size()); }
+	
 	int Compare(const char* Rhs) const	{ return Super::compare(Rhs); }
-	bool StartsWith(const char* Rhs) const { return Super::starts_with(Rhs); }
 	void RemovePrefix(uint32 n)			{ Super::remove_prefix(n); }
 };
 
diff --git a/Engine/Source/Programs/UnrealTraceServer/src/StoreSettings.cpp b/Engine/Source/Programs/UnrealTraceServer/src/StoreSettings.cpp
--- a/Engine/Source/Programs/UnrealTraceServer/src/StoreSettings.cpp
+++ b/Engine/Source/Programs/UnrealTraceServer/src/StoreSettings.cpp
@@ -73,7 +73,7 @@ public:
 	{
 		if (Ident != nullptr)
 		{
-			PRINTCHECKEDF(Handle, "%s=%lld\n", Ident, Value)
+			PRINTCHECKEDF(Handle, "%s=%lld\n", Ident, (long long int) Value)
 			return true;
 		}
 		return false;
@@ -110,7 +110,7 @@ public:
 		FILE* Handle = std::fopen((const char*) Path.string().c_str(), "r");
 		if (!Handle)
 		{
-			TS_SETTINGS_LOG("Unable to open settings file.")
+			TS_SETTINGS_LOGF("Unable to open settings file %s.\n", Path.string().c_str())
 			return; // Safe to return here, buffer and entries will be empty
 		}
 
@@ -122,7 +122,7 @@ public:
 		const size_t Read = std::fread(Buffer.GetData(), sizeof(uint8), FileSize, Handle);
 		if (!Read)
 		{
-			TS_SETTINGS_LOG("Unable to read settings file.")
+			TS_SETTINGS_LOGF("Unable to read settings file %s.\n", Path.string().c_str())
 			std::fclose(Handle);
 			return;
 		}
@@ -333,7 +333,7 @@ void FStoreSettings::ApplySettingsFromCbor(const uint8* Buffer, uint32 NumBytes)
 				if (Context.GetType() == ECborType::String)
 				{
 					FStringView Value = Context.AsString();
-					if (Value.StartsWith("-"))
+					if (Value.Len() > 0 && Value[0] == '-')
 					{
 						Value.RemovePrefix(1);
 						fs::path ValuePath(*FString(Value));
