Fix directory variable

diff --git a/Engine/Source/Programs/UnrealTraceServer/src/Store.cpp b/Engine/Source/Programs/UnrealTraceServer/src/Store.cpp
--- a/Engine/Source/Programs/UnrealTraceServer/src/Store.cpp
+++ b/Engine/Source/Programs/UnrealTraceServer/src/Store.cpp
@@ -247,7 +247,7 @@ FStore::FMount::FMount(FStore* InParent, asio::io_context& InIoContext, const fs
 	DirWatcher = new FDirWatcher(IoContext, DirWatchHandle);
 #elif TS_USING(TS_PLATFORM_LINUX)
 	int inotfd = inotify_init();
-	int watch_desc = inotify_add_watch(inotfd, StoreDir.c_str(), IN_CREATE | IN_DELETE);
+	int watch_desc = inotify_add_watch(inotfd, Dir.c_str(), IN_CREATE | IN_DELETE);
 	DirWatcher = new FDirWatcher(IoContext, inotfd);
 #elif PLATFORM_MAC
 	DirWatcher = new FDirWatcher(*StoreDir);
