Implement store settings creation for Linux

diff --git a/Engine/Source/Programs/UnrealTraceServer/src/Main.cpp b/Engine/Source/Programs/UnrealTraceServer/src/Main.cpp
--- a/Engine/Source/Programs/UnrealTraceServer/src/Main.cpp
+++ b/Engine/Source/Programs/UnrealTraceServer/src/Main.cpp
@@ -289,14 +289,6 @@ struct FLoggingScope
 
 // {{{1 store ------------------------------------------------------------------
 
-////////////////////////////////////////////////////////////////////////////////
-struct FStoreOptions
-{
-	FPath		Dir;
-	int			Port			= 1989;
-	int			RecorderPort	= 1981;
-};
-
 ////////////////////////////////////////////////////////////////////////////////
 static void ParseOptions(int ArgC, char** ArgV, FStoreSettings* Settings)
 {
@@ -1199,19 +1191,19 @@ static int MainDaemon(int ArgC, char** ArgV, pid_t ParentPid)
 
 	// Fire up the store
 	TS_LOG("Starting the store");
-	FStoreService* StoreService;
+	FStoreSettings* Settings = new FStoreSettings();
+	FStoreService* StoreService = nullptr;
 	{
-		FStoreOptions Options = ParseOptions(ArgC, ArgV);
+		FPath HomeDir;
+		GetUnrealTraceHome(HomeDir);
 
-		if (Options.Dir.empty())
-		{
-			FPath StoreDir;
-			GetUnrealTraceHome(StoreDir);
-			StoreDir /= "Store";
-			Options.Dir = StoreDir;
-		}
+		// Read settings from configuration file
+		Settings->ReadFromSettings(HomeDir);
 
-		StoreService = StartStore(Options);
+		// Override with command line arguments
+		ParseOptions(ArgC, ArgV, Settings);
+
+		StoreService = FStoreService::Create(Settings);
 	}
 	OnScopeExit([StoreService] () { delete StoreService; });
 
@@ -1324,6 +1316,10 @@ int main(int ArgC, char** ArgV)
 		return 127;
 	}
 
+#ifdef TS_WITH_MUSL
+	printf("Using musl\n");
+#endif
+
 	struct {
 		const char*	Verb;
 		int			(*Entry)(int, char**);
