Disable system tray on non-windows platforms

diff --git a/Engine/Source/Programs/UnrealTraceServer/src/SystemTray.cpp b/Engine/Source/Programs/UnrealTraceServer/src/SystemTray.cpp
--- a/Engine/Source/Programs/UnrealTraceServer/src/SystemTray.cpp
+++ b/Engine/Source/Programs/UnrealTraceServer/src/SystemTray.cpp
@@ -4,6 +4,8 @@
 #include "StoreService.h"
 #include "Version.h"
 
+#if TS_USING(TS_PLATFORM_WINDOWS)
+
 #define TS_WITH_SYSTEM_TRAY TS_OFF
 
 ////////////////////////////////////////////////////////////////////////////////
@@ -149,4 +151,6 @@ void RemoveFromSystemTray()
 #endif
 }
 
+#endif //TS_PLATFORM_WINDOWS
+
 /* vim: set noexpandtab : */
