Change build system to use xmake on all platforms

In order to unify the build systems on all platforms add makefiles for xmake and deprecate nmake/gmake files. This allows us to use the same configuration files on all platforms and easily add support for using musl on Linux. Updated readme file to reflect the new system.

diff --git a/Engine/Source/Programs/UnrealTraceServer/.gitignore b/Engine/Source/Programs/UnrealTraceServer/.gitignore
--- a/Engine/Source/Programs/UnrealTraceServer/.gitignore
+++ b/Engine/Source/Programs/UnrealTraceServer/.gitignore
@@ -1,4 +1,4 @@
-_build
-_intermediate
+build
+vsxmake*
 tags
 *.user
diff --git a/Engine/Source/Programs/UnrealTraceServer/Readme.md b/Engine/Source/Programs/UnrealTraceServer/Readme.md
--- a/Engine/Source/Programs/UnrealTraceServer/Readme.md
+++ b/Engine/Source/Programs/UnrealTraceServer/Readme.md
@@ -1,28 +1,54 @@
+
 # Unreal Trace Server
 
-Standalone hub for recording traces from applications and serving them to analyzers (e.g. Unreal Insights).
+Unreal Trace Server acts as a hub between runtimes that are tracing performance instrumentation and tools like Unreal Insights that consume and present that data for analysis. By default TCP ports 1981 and 1989 are used, where the former receives trace data, and the latter is used by tools to query the server's store.
+
+The server stores new traces to the "store directory". It can also watch for trace files in "watch directories".
+
+The server stores persistent configuration files in:
+* `%USERPROFILE%/UnrealEngine/Common/UnrealTrace` on Windows
+* `~/UnrealEngine/UnrealTrace` on Linux and MacOS.
+
+# Building and running
+
+Unreal Trace Server uses [xmake](https://xmake.io) for building and generating solution files. This gives us many benefits:
 
-# Building
+ * Unified build files across platforms
+ * Ability to easily use custom toolchains (like musl)
+ * Potentially remove third party source and rely on xmakes package manager
+ * Generate IDE files (Visual Studio, XCode, etc)
 
-## Windows
+To build, you first need to install or download xmake to your machine. Follow the [installation guide](https://xmake.io/#/guide/installation) for your host platform. It is also possible to run xmake in a standalone mode (without installing) as long as xmake is in the `PATH` environment variable.
 
-1. Run "Developer Command Prompt for Visual Studio XXXX" from the start menu.
-2. Start a x64 version of the command line by executing `VsDevCmd.bat -host_arch=amd64 -arch=amd64`
-3. Navigate to this folder
-4. Execute command 
-	* `nmake` for building release configuration.
-	* `nmake /E DEBUG=1` for debug configuration.
+Once installed navigate to the project directory. To build the server run:
+```
+> xmake
+```
+The output files are emitted into `build/[platform]/[architecture]`.
 
-Visual Studio solution is provided for convenience. Note that changes that affect `Pch.h` requires a clean since there is no dependency checking for the precompiled header.
+### Running
+Unreal Trace Server requires a command to run.
+* `fork` - Starts a background server, upgrading any existing instance.
+* `daemon` - The mode that a background server runs in.
+* `kill` - Shuts down a currently running instance.
+
+As long a no other instance of the server is currently running, start the server by: 
+```
+>  xmake run UnrealTraceServer daemon
+```
+
+
+### IDE files
+Optionally IDE files can be generated, for example Visual Studio solution, by using:
+```
+> xmake project -k vsxmake -m "debug,release"
+```
 
 # Making a release
 
-Bump `TS_VERSION_MINOR` so the auto-update mechanisms activate when users receive the newer version.
+1. Bump `TS_VERSION_MINOR` so the auto-update mechanisms activate when users receive the newer version. 
+2. Make sure release mode is active (default) by running `xmake config -m release`.
+3. Run the install script `xmake install`. 
+
+The install script will check out the binaries deployed in the Unreal Engine tree. Submit the binaries together with the `Version.h`.
 
-## Windows
-1. Run "Developer Command Prompt for Visual Studio XXXX" from the start menu.
-2. Start a x64 version of the command line by executing `VsDevCmd.bat -host_arch=amd64 -arch=amd64`
-3. Navigate to this folder
-4. Execute command 
-	* `nmake clean release`
-5. In Perforce check in new version of `UnrealTraceServer.exe` and the corresponding .PDB file.
\ No newline at end of file
diff --git a/Engine/Source/Programs/UnrealTraceServer/makefile b/Engine/Source/Programs/UnrealTraceServer/makefile
deleted file mode 100644
