Fix directory watcher on Linux/Mac

* Missed call to parent refresh
* Using incorrect platform define

diff --git a/Engine/Source/Programs/UnrealTraceServer/src/Store.cpp b/Engine/Source/Programs/UnrealTraceServer/src/Store.cpp
--- a/Engine/Source/Programs/UnrealTraceServer/src/Store.cpp
+++ b/Engine/Source/Programs/UnrealTraceServer/src/Store.cpp
@@ -11,6 +11,9 @@
 #	include <CoreServices/CoreServices.h>
 #endif
 
+#ifndef TS_DEBUG_POSIX_WATCHER_EVENTS
+	#define TS_DEBUG_POSIX_WATCHER_EVENTS TS_OFF
+#endif
 
 ////////////////////////////////////////////////////////////////////////////////
 // Pre-C++20 there is now way to convert between clocks. C++20 onwards it is
@@ -44,6 +47,53 @@ public:
 	void async_wait(HandlerType InHandler)
 	{
 		asio::posix::stream_descriptor::async_wait(asio::posix::stream_descriptor::wait_read, InHandler);
+
+		// Some new events have come into the stream. We need to read the events
+		// event if we are not acting on the events themself (see Refresh for platform
+		// independent update). If we don't read the data the next call to wait will
+		// trigger immediately.
+		bytes_readable AvailableCmd(true);
+		io_control(AvailableCmd);
+		size_t AvailableBytes = AvailableCmd.get();
+		uint8* Buffer = (uint8*)malloc(AvailableBytes);
+		read_some(asio::buffer(Buffer, AvailableBytes));
+		
+		#if TS_USING(TS_DEBUG_POSIX_WATCHER_EVENTS)
+		size_t Cursor = 0;
+		while(Cursor < AvailableBytes)
+		{
+			inotify_event *Event = (inotify_event *)Buffer + Cursor;
+			printf("Recieved file event (0x08x) ", Event->cookie);
+			if (Event->len > 0)
+				printf("on '%s': ", Event->name);
+			if (Event->mask & IN_ACCESS)
+				printf("ACCESS ");
+			if ((Event->mask & IN_ATTRIB) != 0)
+				printf("ATTRIB ");
+			if ((Event->mask & IN_CLOSE_WRITE) != 0)
+				printf("CLOSE_WRITE ");
+			if ((Event->mask & IN_CLOSE_NOWRITE) != 0)
+				printf("CLOSE_NOWRITE ");
+			if ((Event->mask & IN_CREATE) != 0)
+				printf("CREATE ");
+			if ((Event->mask & IN_DELETE) != 0)
+				printf("DELETE ");
+			if ((Event->mask & IN_DELETE_SELF) != 0)
+				printf("DELETE_SELF ");
+			if ((Event->mask & IN_MODIFY) != 0)
+				printf("MODIFY ");
+			if ((Event->mask & IN_MOVE_SELF) != 0)
+				printf("MOVE_SELF ");
+			if ((Event->mask & IN_MOVED_FROM) != 0)
+				printf("MOVED_FROM ");
+			if ((Event->mask & IN_MOVED_TO) != 0)
+				printf("MOVED_TO ");
+			if ((Event->mask & IN_OPEN) != 0)
+				printf("OPEN ");
+			printf("\n");
+			Cursor += sizeof(inotify_event) + Event->len;
+		}
+		#endif
 	}
 };
 
@@ -249,8 +299,8 @@ FStore::FMount::FMount(FStore* InParent, asio::io_context& InIoContext, const fs
 	int inotfd = inotify_init();
 	int watch_desc = inotify_add_watch(inotfd, Dir.c_str(), IN_CREATE | IN_DELETE);
 	DirWatcher = new FDirWatcher(IoContext, inotfd);
-#elif PLATFORM_MAC
-	DirWatcher = new FDirWatcher(*StoreDir);
+#elif TS_USING(TS_PLATFORM_MAC)
+	DirWatcher = new FDirWatcher(Dir.c_str());
 #endif
 
 	WatchDir();
@@ -472,7 +522,7 @@ void FStore::FMount::WatchDir()
 			WatchDir();
 		});
 #else
-		Refresh();
+		Parent->Refresh();
 		WatchDir();
 #endif
 	});
