Use entire path for id

Use entire path to calculate the id. Just using the stem hides files that happens to be present in multiple mounts.

diff --git a/Engine/Source/Programs/UnrealTraceServer/src/Store.cpp b/Engine/Source/Programs/UnrealTraceServer/src/Store.cpp
--- a/Engine/Source/Programs/UnrealTraceServer/src/Store.cpp
+++ b/Engine/Source/Programs/UnrealTraceServer/src/Store.cpp
@@ -205,7 +205,7 @@ FStore::FTrace::FTrace(const FPath& InPath)
 : Path(InPath)
 {
 	const FString Name = GetName();
-	Id = QuickStoreHash(*Name);
+	Id = QuickStoreHash(InPath.c_str());
 
 	std::error_code Ec;
 	// Calculate that trace's timestamp. Bias in seconds then convert to 0.1us.
