Parse session and trace guids

diff --git a/Engine/Source/Programs/UnrealTraceServer/src/Foundation.h b/Engine/Source/Programs/UnrealTraceServer/src/Foundation.h
--- a/Engine/Source/Programs/UnrealTraceServer/src/Foundation.h
+++ b/Engine/Source/Programs/UnrealTraceServer/src/Foundation.h
@@ -159,4 +159,11 @@ private:
 };
 #endif // TS_PLATFORM_WINDOWS
 
+////////////////////////////////////////////////////////////////////////////////
+
+struct FGuid
+{
+	uint32 Bits[4];
+};
+
 /* vim: set noexpandtab : */
diff --git a/Engine/Source/Programs/UnrealTraceServer/src/Recorder.cpp b/Engine/Source/Programs/UnrealTraceServer/src/Recorder.cpp
--- a/Engine/Source/Programs/UnrealTraceServer/src/Recorder.cpp
+++ b/Engine/Source/Programs/UnrealTraceServer/src/Recorder.cpp
@@ -35,6 +35,8 @@ private:
 	uint8*				PreambleCursor;
 	uint32				TraceId = 0;
 	uint16				ControlPort = 0;
+	FGuid				SessionGuid;
+	FGuid				TraceGuid;
 	uint8				FillDrainCounter = 1;
 	uint32				BufferPolarity = 0;
 	uint8*				Buffer[2];
@@ -223,9 +225,17 @@ bool FRecorderRelay::ReadMetadata(int32 Size)
 			return false;
 		}
 
-		if (MetadataField.Id == 0) /* ControlPortFieldId */
+		switch (MetadataField.Id) 
 		{
+		case 0: /* ControlPortFieldId */ 
 			ControlPort = *(const uint16*)Cursor;
+			break;
+		case 1: /* SessionGuid */ 
+			SessionGuid = *(const FGuid*)Cursor;
+			break;
+		case 2: /* TraceGuid */
+			TraceGuid = *(const FGuid*)Cursor;
+			break;
 		}
 
 		Cursor += MetadataField.Size;
