Parse Guids from string representation

diff --git a/Engine/Source/Programs/UnrealTraceServer/src/Foundation.h b/Engine/Source/Programs/UnrealTraceServer/src/Foundation.h
--- a/Engine/Source/Programs/UnrealTraceServer/src/Foundation.h
+++ b/Engine/Source/Programs/UnrealTraceServer/src/Foundation.h
@@ -163,7 +163,41 @@ private:
 
 struct FGuid
 {
-	uint32 Bits[4];
+	uint32 Bits[4] = { 0, 0, 0, 0 };
+
+	static bool ParseGuid(FStringView Source, FGuid& OutGuid)
+	{
+		std::basic_string_view<char> SourceView(Source.GetData(), Source.Len());
+		std::basic_string_view<char> A(SourceView.substr(0, 8)),
+			B(SourceView.substr(8, 8)),
+			C(SourceView.substr(16, 8)),
+			D(SourceView.substr(24, 8));
+
+		if (auto result = std::from_chars(A.data(), A.data() + A.length(), OutGuid.Bits[0], 16); result.ptr != (A.data() + A.length()))
+		{
+			return false;
+		}
+		if (auto result = std::from_chars(B.data(), B.data() + B.length(), OutGuid.Bits[1], 16); result.ptr != (B.data() + B.length()))
+		{
+			return false;
+		}
+		if (auto result = std::from_chars(C.data(), C.data() + C.length(), OutGuid.Bits[2], 16); result.ptr != (C.data() + C.length()))
+		{
+			return false;
+		}
+		if (auto result = std::from_chars(D.data(), D.data() + D.length(), OutGuid.Bits[3], 16); result.ptr != (D.data() + D.length()))
+		{
+			return false;
+		}
+		
+		return true;
+	}
+
+	static bool Equal(const FGuid& Rhs, const FGuid Lhs)
+	{
+		return Rhs.Bits[0] == Lhs.Bits[0] && Rhs.Bits[1] == Lhs.Bits[1] 
+			&& Rhs.Bits[2] == Lhs.Bits[2] && Rhs.Bits[3] == Lhs.Bits[3];
+	}
 };
 
 /* vim: set noexpandtab : */
diff --git a/Engine/Source/Programs/UnrealTraceServer/src/Pch.h b/Engine/Source/Programs/UnrealTraceServer/src/Pch.h
--- a/Engine/Source/Programs/UnrealTraceServer/src/Pch.h
+++ b/Engine/Source/Programs/UnrealTraceServer/src/Pch.h
@@ -60,12 +60,13 @@
 
 #include <algorithm>
 #include <atomic>
+#include <charconv>
 #include <chrono>
 #include <cstdint>
 #include <ctime>
 #include <filesystem>
-#include <string>
 #include <string_view>
+#include <string>
 #include <vector>
 
 #if TS_USING(TS_PLATFORM_WINDOWS)
