Implement getting session info using trace guid.

diff --git a/Engine/Source/Programs/UnrealTraceServer/src/Recorder.cpp b/Engine/Source/Programs/UnrealTraceServer/src/Recorder.cpp
--- a/Engine/Source/Programs/UnrealTraceServer/src/Recorder.cpp
+++ b/Engine/Source/Programs/UnrealTraceServer/src/Recorder.cpp
@@ -22,6 +22,8 @@ public:
 	uint32				GetTraceId() const;
 	uint32				GetIpAddress() const;
 	uint32				GetControlPort() const;
+	const FGuid&		GetSessionGuid() const;
+	const FGuid&		GetTraceGuid() const;
 
 private:
 	virtual void		OnIoComplete(uint32 Id, int32 Size) override;
@@ -151,6 +153,18 @@ uint32 FRecorderRelay::GetControlPort() const
 	return ControlPort;
 }
 
+////////////////////////////////////////////////////////////////////////////////
+const FGuid& FRecorderRelay::GetSessionGuid() const
+{
+	return SessionGuid;
+}
+
+////////////////////////////////////////////////////////////////////////////////
+const FGuid& FRecorderRelay::GetTraceGuid() const
+{
+	return TraceGuid;
+}
+
 ////////////////////////////////////////////////////////////////////////////////
 bool FRecorderRelay::CreateTrace()
 {
@@ -341,6 +355,18 @@ uint32 FRecorder::FSession::GetControlPort() const
 	return Relay->GetControlPort();
 }
 
+////////////////////////////////////////////////////////////////////////////////
+const FGuid& FRecorder::FSession::GetSessionGuid() const
+{
+	return Relay->GetSessionGuid();
+}
+
+////////////////////////////////////////////////////////////////////////////////
+const FGuid& FRecorder::FSession::GetTraceGuid() const
+{
+	return Relay->GetTraceGuid();
+}
+
 
 
 ////////////////////////////////////////////////////////////////////////////////
diff --git a/Engine/Source/Programs/UnrealTraceServer/src/Recorder.h b/Engine/Source/Programs/UnrealTraceServer/src/Recorder.h
--- a/Engine/Source/Programs/UnrealTraceServer/src/Recorder.h
+++ b/Engine/Source/Programs/UnrealTraceServer/src/Recorder.h
@@ -23,6 +23,8 @@ public:
 		uint32				GetTraceId() const;
 		uint32				GetIpAddress() const;
 		uint32				GetControlPort() const;
+		const FGuid&		GetSessionGuid() const;
+		const FGuid&		GetTraceGuid() const;
 
 	private:
 		friend				FRecorder;
diff --git a/Engine/Source/Programs/UnrealTraceServer/src/StoreCborServer.cpp b/Engine/Source/Programs/UnrealTraceServer/src/StoreCborServer.cpp
--- a/Engine/Source/Programs/UnrealTraceServer/src/StoreCborServer.cpp
+++ b/Engine/Source/Programs/UnrealTraceServer/src/StoreCborServer.cpp
@@ -116,6 +116,23 @@ void FStoreCborPeer::OnSessionInfo()
 			}
 		}
 	}
+	else if (FStringView TraceGuidStr = Response.GetString("trace_guid", ""); TraceGuidStr.Len() > 0)
+	{
+		FGuid TraceGuid;
+		if (FGuid::ParseGuid(TraceGuidStr, TraceGuid))
+		{
+			for (int i = 0, n = Recorder.GetSessionCount(); i < n; ++i)
+			{
+				const FRecorder::FSession* Candidate = Recorder.GetSessionInfo(i);
+
+					if (FGuid::Equal(Candidate->GetTraceGuid(), TraceGuid))
+					{
+						Session = Candidate;
+						break;
+					}
+			}
+		}
+	}
 
 	if (Session == nullptr)
 	{
