Update readme to reflect dynamic ports used

#prefligh none

diff --git a/Engine/Source/Programs/UnrealTraceServer/Readme.md b/Engine/Source/Programs/UnrealTraceServer/Readme.md
--- a/Engine/Source/Programs/UnrealTraceServer/Readme.md
+++ b/Engine/Source/Programs/UnrealTraceServer/Readme.md
@@ -9,6 +9,14 @@ The server stores persistent configuration files in:
 * `%USERPROFILE%/UnrealEngine/Common/UnrealTrace` on Windows
 * `~/UnrealEngine/UnrealTrace` on Linux and MacOS.
 
+## Connecting to the server
+
+When running UnrealTraceServer as a centralized server make sure the following ports are available to clients:
+
+ * Recorder port (default 1981)
+ * Store port (default 1989)
+ * Dynamic ports (see operating system configuration)
+
 # Building and running
 
 Unreal Trace Server uses [xmake](https://xmake.io) for building and generating solution files. This gives us many benefits:
