New UTS version

* Parse Session guid and trace guid from metadata if available for live sessions
* Add api to fetch a session based on above guid.

diff --git a/Engine/Source/Programs/UnrealTraceServer/src/Version.h b/Engine/Source/Programs/UnrealTraceServer/src/Version.h
--- a/Engine/Source/Programs/UnrealTraceServer/src/Version.h
+++ b/Engine/Source/Programs/UnrealTraceServer/src/Version.h
@@ -7,7 +7,7 @@
 #endif
 
 #if !defined(TS_VERSION_MINOR)
-#	define TS_VERSION_MINOR		15
+#	define TS_VERSION_MINOR		16
 #endif
 
 #define TS_ICON_ID				493
