Change directory watcher on MacOs in UnrealTraceServer

Use dispatch queues and updated watch logic.

diff --git a/Engine/Source/Programs/UnrealTraceServer/src/Store.cpp b/Engine/Source/Programs/UnrealTraceServer/src/Store.cpp
--- a/Engine/Source/Programs/UnrealTraceServer/src/Store.cpp
+++ b/Engine/Source/Programs/UnrealTraceServer/src/Store.cpp
@@ -11,8 +11,8 @@
 #	include <CoreServices/CoreServices.h>
 #endif
 
-#ifndef TS_DEBUG_POSIX_WATCHER_EVENTS
-	#define TS_DEBUG_POSIX_WATCHER_EVENTS TS_OFF
+#ifndef TS_DEBUG_FS_EVENTS
+	#define TS_DEBUG_FS_EVENTS TS_OFF
 #endif
 
 ////////////////////////////////////////////////////////////////////////////////
@@ -58,7 +58,7 @@ public:
 		uint8* Buffer = (uint8*)malloc(AvailableBytes);
 		read_some(asio::buffer(Buffer, AvailableBytes));
 		
-		#if TS_USING(TS_DEBUG_POSIX_WATCHER_EVENTS)
+		#if TS_USING(TS_DEBUG_FS_EVENTS)
 		size_t Cursor = 0;
 		while(Cursor < AvailableBytes)
 		{
@@ -107,6 +107,10 @@ public:
 	{
 		std::error_code ErrorCode;
 		StoreDir = std::filesystem::absolute(InStoreDir, ErrorCode);
+		// Create watcher queue
+		char WatcherName[128];
+		snprintf(WatcherName, sizeof(WatcherName), "%s-%p", "FileWatcher", this);
+		DispatchQueue = dispatch_queue_create(WatcherName, DISPATCH_QUEUE_SERIAL);
 	}
 	void async_wait(HandlerType InHandler);
 	void cancel()
@@ -118,9 +122,9 @@ public:
 		if (bIsRunning)
 		{
 			FSEventStreamStop(EventStream);
-			FSEventStreamUnscheduleFromRunLoop(EventStream, CFRunLoopGetCurrent(), kCFRunLoopDefaultMode);
-			FSEventStreamInvalidate(EventStream);
+			FSEventStreamInvalidate(EventStream); // Also removed from dispatch queue
 			FSEventStreamRelease(EventStream);
+			dispatch_release(DispatchQueue);
 			bIsRunning = false;
 		}
 	}
@@ -128,18 +132,52 @@ public:
 
 	void ProcessChanges(size_t EventCount, void* EventPaths, const FSEventStreamEventFlags EventFlags[])
 	{
-		Handler(asio::error_code());
+		bool bWatchedEvent = false;
+		const char** EventPathsArray = (const char**) EventPaths;
+		for (size_t EventIdx = 0; EventIdx < EventCount; ++EventIdx)
+		{
+			const char* Path = EventPathsArray[EventIdx];
+			const FSEventStreamEventFlags& Flags = EventFlags[EventIdx];
+#if TS_USING(TS_DEBUG_FS_EVENTS)
+			printf("Recieved file event (%d) ", EventIdx);
+			if (Path != nullptr)
+				printf("on '%s': ", Path);
+			else
+				printf("on unknown file: ");
+			if (Flags & kFSEventStreamEventFlagItemCreated)
+				printf(" CREATED");
+			if (Flags & kFSEventStreamEventFlagItemRemoved)
+				printf(" REMOVED");
+			if (Flags & kFSEventStreamEventFlagItemRenamed)
+				printf(" RENAMED");
+			printf("\n");
+#endif
+			constexpr unsigned int InterestingFlags = kFSEventStreamEventFlagItemCreated | 
+				kFSEventStreamEventFlagItemRemoved | kFSEventStreamEventFlagItemRenamed;
+			bWatchedEvent |= !!(Flags & InterestingFlags);
+		}
+		if (bWatchedEvent)
+		{
+			Handler(asio::error_code());
+		}
 	}
 
 	FSEventStreamRef	EventStream;
-	HandlerType Handler;
+	dispatch_queue_t	DispatchQueue;
+	HandlerType 		Handler;
 private:
+	static void MacCallback(ConstFSEventStreamRef StreamRef,
+					void* InDirWatcherPtr,
+					size_t EventCount,
+					void* EventPaths,
+					const FSEventStreamEventFlags EventFlags[],
+					const FSEventStreamEventId EventIDs[]);
 	bool bIsRunning = false;
 	FPath StoreDir;
 };
 
 ////////////////////////////////////////////////////////////////////////////////
-static void MacCallback(ConstFSEventStreamRef StreamRef,
+void FStore::FMount::FDirWatcher::MacCallback(ConstFSEventStreamRef StreamRef,
 					void* InDirWatcherPtr,
 					size_t EventCount,
 					void* EventPaths,
@@ -171,28 +209,38 @@ void FStore::FMount::FDirWatcher::async_wait(HandlerType InHandler)
 	Context.copyDescription = NULL;
 
 	// Set up streaming and turn it on
-	std::string Path = StoreDir;
-	CFStringRef FullPathMac = CFStringCreateWithBytes(
-		kCFAllocatorDefault,
-		(const uint8*)Path.data(),
-		Path.size(),
-		kCFStringEncodingUnicode,
-		false);
-
+	CFStringRef FullPathMac = CFStringCreateWithFileSystemRepresentation(
+		NULL,
+		(const char*)StoreDir.c_str());
 	CFArrayRef PathsToWatch = CFArrayCreate(NULL, (const void**)&FullPathMac, 1, NULL);
 
-	EventStream = FSEventStreamCreate(NULL,
+	// Create the event stream object
+	EventStream = FSEventStreamCreate(
+		NULL,
 		&MacCallback,
 		&Context,
 		PathsToWatch,
 		kFSEventStreamEventIdSinceNow,
 		Latency,
-		kFSEventStreamCreateFlagUseCFTypes | kFSEventStreamCreateFlagNoDefer
+		kFSEventStreamCreateFlagNoDefer|kFSEventStreamCreateFlagFileEvents
 	);
 
-	FSEventStreamScheduleWithRunLoop(EventStream, CFRunLoopGetCurrent(), kCFRunLoopDefaultMode);
-	FSEventStreamStart(EventStream);
-	bIsRunning = true;
+	if (EventStream == nullptr)
+	{
+		printf("Failed to create file event stream for %s\n", CFStringGetCStringPtr(FullPathMac, kCFStringEncodingUnicode));
+	}
+
+	FSEventStreamSetDispatchQueue(EventStream, DispatchQueue);
+	bIsRunning = FSEventStreamStart(EventStream);
+
+	if (bIsRunning) 
+	{
+		printf("Watcher enabled on %s\n", StoreDir.c_str());
+	}
+	else
+	{
+		printf("Failed to start watcher for %s\n", StoreDir.c_str());
+	}
 
 	Handler = InHandler;
 }
