Update language in status message

diff --git a/Engine/Source/Programs/UnrealTraceServer/src/Main.cpp b/Engine/Source/Programs/UnrealTraceServer/src/Main.cpp
--- a/Engine/Source/Programs/UnrealTraceServer/src/Main.cpp
+++ b/Engine/Source/Programs/UnrealTraceServer/src/Main.cpp
@@ -1018,7 +1018,7 @@ static int MainKill(int ArgC, char** ArgV)
 	{
 		if (errno == ENOENT)
 		{
-			TS_LOG("All good. Ain't nuffin' running me ol' mucker.");
+			TS_LOG("All good. Ain't nuffin' running me ol' pal.");
 			return Result_Ok;
 		}
 
