Add version query endpoint in UTS

diff --git a/Engine/Source/Programs/UnrealTraceServer/src/StoreCborServer.cpp b/Engine/Source/Programs/UnrealTraceServer/src/StoreCborServer.cpp
--- a/Engine/Source/Programs/UnrealTraceServer/src/StoreCborServer.cpp
+++ b/Engine/Source/Programs/UnrealTraceServer/src/StoreCborServer.cpp
@@ -4,12 +4,13 @@
 #include "AsioIoable.h"
 #include "AsioSocket.h"
 #include "CborPayload.h"
-#include "StoreService.h"
 #include "Recorder.h"
 #include "Store.h"
 #include "StoreCborServer.h"
-#include "TraceRelay.h"
+#include "StoreService.h"
 #include "StoreSettings.h"
+#include "TraceRelay.h"
+#include "Version.h"
 
 ////////////////////////////////////////////////////////////////////////////////
 class FStoreCborPeer
@@ -26,6 +27,7 @@ protected:
 	void					OnSessionCount();
 	void					OnSessionInfo();
 	void					OnStatus();
+	void					OnVersion();
 	void					OnTraceCount();
 	void					OnTraceInfo();
 	void					OnTraceRead();
@@ -165,6 +167,21 @@ void FStoreCborPeer::OnStatus()
 	SendResponse(Builder.Done());
 }
 
+///////////////////////////////////////////////////////////////////////////////
+void FStoreCborPeer::OnVersion()
+{
+	TPayloadBuilder<> Builder(EStatusCode::Success);
+	Builder.AddInteger("major", TS_VERSION_PROTOCOL);
+	Builder.AddInteger("minor", TS_VERSION_MINOR);
+#if TS_USING(TS_BUILD_DEBUG)
+	const char* Configuration = "Debug";
+#else
+	const char* Configuration = "Release";
+#endif
+	Builder.AddString("configuration", Configuration);
+	SendResponse(Builder.Done());
+}
+
 ////////////////////////////////////////////////////////////////////////////////
 void FStoreCborPeer::OnTraceCount()
 {
@@ -278,6 +295,7 @@ void FStoreCborPeer::OnPayload()
 		{ QuickStoreHash("v1/session/count"),	&FStoreCborPeer::OnSessionCount },
 		{ QuickStoreHash("v1/session/info"),	&FStoreCborPeer::OnSessionInfo },
 		{ QuickStoreHash("v1/status"),			&FStoreCborPeer::OnStatus },
+		{ QuickStoreHash("v1/version"),			&FStoreCborPeer::OnVersion },
 		{ QuickStoreHash("v1/trace/count"),		&FStoreCborPeer::OnTraceCount },
 		{ QuickStoreHash("v1/trace/info"),		&FStoreCborPeer::OnTraceInfo },
 		{ QuickStoreHash("v1/trace/read"),		&FStoreCborPeer::OnTraceRead },
