Fix socket keep alive in UTS for Mac/Linux

On Windows keepalive options were set on the recording socket,  in order to detect if a tracing process suddenly closes it's connection (it's killed or crashes).  These options were missing from Mac and Linux configurations.

diff --git a/Engine/Source/Programs/UnrealTraceServer/src/Recorder.cpp b/Engine/Source/Programs/UnrealTraceServer/src/Recorder.cpp
--- a/Engine/Source/Programs/UnrealTraceServer/src/Recorder.cpp
+++ b/Engine/Source/Programs/UnrealTraceServer/src/Recorder.cpp
@@ -92,6 +92,18 @@ FRecorderRelay::FRecorderRelay(asio::ip::tcp::socket& Socket, FStore& InStore)
 		nullptr,
 		nullptr
 	);
+#elif TS_USING(TS_PLATFORM_MAC)
+    int Enabled = 1, Secs = 3;
+    
+    setsockopt(Socket.native_handle(), SOL_SOCKET,  SO_KEEPALIVE, &Enabled, sizeof(Enabled));
+    setsockopt(Socket.native_handle(), IPPROTO_TCP, TCP_KEEPALIVE, &Secs, sizeof(Secs));
+#elif TS_USING(TS_PLATFORM_LINUX)
+    int Enabled = 1, KeepCnt = 5, KeepIdle = 3, KeepIntvl = 1;
+
+    setsockopt(Socket.native_handle(), SOL_SOCKET,  SO_KEEPALIVE, &Enabled, sizeof Enabled);
+    setsockopt(Socket.native_handle(), IPPROTO_TCP, TCP_KEEPCNT, &KeepCnt, sizeof(int));
+    setsockopt(Socket.native_handle(), IPPROTO_TCP, TCP_KEEPIDLE, &KeepIdle, sizeof(int));
+    setsockopt(Socket.native_handle(), IPPROTO_TCP, TCP_KEEPINTVL, &KeepIntvl, sizeof(int));
 #endif
 
 	// Kick things off by reading the magic four bytes at the start of the stream
