Version 1.17 of UnrealTraceServer for windows

UnrealTraceServer executable and debug files.
Version: 1.17
Platform: windows, mac, linux
Architecture: x64
Compiled with musl: false (linux: true)

diff --git a/Engine/Source/Programs/UnrealTraceServer/src/Version.h b/Engine/Source/Programs/UnrealTraceServer/src/Version.h
--- a/Engine/Source/Programs/UnrealTraceServer/src/Version.h
+++ b/Engine/Source/Programs/UnrealTraceServer/src/Version.h
@@ -7,7 +7,7 @@
 #endif
 
 #if !defined(TS_VERSION_MINOR)
-#	define TS_VERSION_MINOR		16
+#	define TS_VERSION_MINOR		17
 #endif
 
 #define TS_ICON_ID				493
