Fix directory listening for Linux

The file notification api was not told to listen to renames.

diff --git a/Engine/Source/Programs/UnrealTraceServer/src/Store.cpp b/Engine/Source/Programs/UnrealTraceServer/src/Store.cpp
--- a/Engine/Source/Programs/UnrealTraceServer/src/Store.cpp
+++ b/Engine/Source/Programs/UnrealTraceServer/src/Store.cpp
@@ -63,9 +63,11 @@ public:
 		while(Cursor < AvailableBytes)
 		{
 			inotify_event *Event = (inotify_event *)Buffer + Cursor;
-			printf("Recieved file event (0x08x) ", Event->cookie);
-			if (Event->len > 0)
+			printf("Recieved file event (0x%08x) ", Event->cookie);
+			if (Event->len > 0 && strlen(Event->name))
 				printf("on '%s': ", Event->name);
+			else
+				printf("on the directory ");
 			if (Event->mask & IN_ACCESS)
 				printf("ACCESS ");
 			if ((Event->mask & IN_ATTRIB) != 0)
@@ -345,7 +347,7 @@ FStore::FMount::FMount(FStore* InParent, asio::io_context& InIoContext, const fs
 	DirWatcher = new FDirWatcher(IoContext, DirWatchHandle);
 #elif TS_USING(TS_PLATFORM_LINUX)
 	int inotfd = inotify_init();
-	int watch_desc = inotify_add_watch(inotfd, Dir.c_str(), IN_CREATE | IN_DELETE);
+	int watch_desc = inotify_add_watch(inotfd, Dir.c_str(), IN_CREATE | IN_DELETE | IN_MOVED_FROM | IN_MOVED_TO);
 	DirWatcher = new FDirWatcher(IoContext, inotfd);
 #elif TS_USING(TS_PLATFORM_MAC)
 	DirWatcher = new FDirWatcher(Dir.c_str());
