Fix UTS timestamps on Linux

UTS was relying on standard library file times to get a common epoch based timestamp. However for Linux this does not seem to produce reliable results, so we fall back to the posix interface functions.

diff --git a/Engine/Source/Programs/UnrealTraceServer/src/Store.cpp b/Engine/Source/Programs/UnrealTraceServer/src/Store.cpp
--- a/Engine/Source/Programs/UnrealTraceServer/src/Store.cpp
+++ b/Engine/Source/Programs/UnrealTraceServer/src/Store.cpp
@@ -6,6 +6,8 @@
 
 #if TS_USING(TS_PLATFORM_LINUX)
 #   include <sys/inotify.h>
+#	include <sys/types.h>
+#	include <sys/stat.h>
 #endif
 #if TS_USING(TS_PLATFORM_MAC)
 #	include <CoreServices/CoreServices.h>
@@ -257,11 +259,25 @@ FStore::FTrace::FTrace(const FPath& InPath)
 	const FString Name = GetName();
 	Id = QuickStoreHash(InPath.c_str());
 
-	std::error_code Ec;
 	// Calculate that trace's timestamp. Bias in seconds then convert to 0.1us.
+	// Note: Asking for std::filesystem timestamp with epoch seems to work on Windows
+	// and Mac, but behaves badly on Linux. Since in C++17 get file time epoch is
+	// undefined in the standard we fall back to stat interface on Linux.
+#if TS_USING(TS_PLATFORM_LINUX)
+	struct stat FileStats;
+	if (stat(Path.c_str(), &FileStats) || FileStats.st_mtime < 0)
+	{
+		// Failure code with errno
+		Timestamp = 0;
+		return;
+	}
+	Timestamp = FileStats.st_mtime;
+#else
+	std::error_code Ec;
 	std::filesystem::file_time_type LastWriteTime = std::filesystem::last_write_time(Path, Ec);
 	auto LastWriteDuration = LastWriteTime.time_since_epoch();
 	Timestamp = std::chrono::duration_cast<std::chrono::seconds>(LastWriteDuration).count();
+#endif
 	Timestamp += FsToUnrealEpochBiasSeconds;
 	Timestamp *= 10'000'000;
 }
