Clean up dispatch queues for recreated watchers

Whenever settings would change we deleted the directory watches for all mounts. On Mac each watcher had a dispatch queue with running jobs. When the those lingering jobs would get callbacks, it would try to access the existing mount instance which would have been deleted.

diff --git a/Engine/Source/Programs/UnrealTraceServer/src/Store.cpp b/Engine/Source/Programs/UnrealTraceServer/src/Store.cpp
--- a/Engine/Source/Programs/UnrealTraceServer/src/Store.cpp
+++ b/Engine/Source/Programs/UnrealTraceServer/src/Store.cpp
@@ -116,6 +116,12 @@ public:
 		snprintf(WatcherName, sizeof(WatcherName), "%s-%p", "FileWatcher", this);
 		DispatchQueue = dispatch_queue_create(WatcherName, DISPATCH_QUEUE_SERIAL);
 	}
+
+	~FDirWatcher()
+	{
+		cancel();
+	}
+
 	void async_wait(HandlerType InHandler);
 	void cancel()
 	{
