Break out logging and instance info

Prepare for future work by breaking out logging and instance info into separate files.

diff --git a/Engine/Source/Programs/UnrealTraceServer/src/InstanceInfo.cpp b/Engine/Source/Programs/UnrealTraceServer/src/InstanceInfo.cpp
new file mode 100644
--- /dev/null
+++ b/Engine/Source/Programs/UnrealTraceServer/src/InstanceInfo.cpp
@@ -0,0 +1,43 @@
+// Copyright Epic Games, Inc. All Rights Reserved.
+
+#include "Pch.h"
+#include "InstanceInfo.h"
+
+////////////////////////////////////////////////////////////////////////////////
+void FInstanceInfo::Set()
+{
+	Version = CurrentVersion;
+#if TS_USING(TS_PLATFORM_WINDOWS)
+	Pid = GetCurrentProcessId();
+#elif TS_USING(TS_PLATFORM_LINUX) || TS_USING(TS_PLATFORM_MAC)
+	Pid = getpid();
+#endif
+	Published.fetch_add(1, std::memory_order_release);
+}
+
+////////////////////////////////////////////////////////////////////////////////
+void FInstanceInfo::WaitForReady() const
+{
+	// Spin until this instance info is published (by another process)
+#if TS_USING(TS_PLATFORM_WINDOWS)
+	for (;; Sleep(0))
+#else
+	for (;; sched_yield())
+#endif
+	{
+		if (Published.load(std::memory_order_acquire))
+		{
+			break;
+		}
+	}
+}
+
+////////////////////////////////////////////////////////////////////////////////
+bool FInstanceInfo::IsOlder() const
+{
+	// Decide which is older; this compiled code or the instance we have a
+	// pointer to.
+	bool bIsOlder = false;
+	bIsOlder |= (Version < FInstanceInfo::CurrentVersion);
+	return bIsOlder;
+}
\ No newline at end of file
diff --git a/Engine/Source/Programs/UnrealTraceServer/src/InstanceInfo.h b/Engine/Source/Programs/UnrealTraceServer/src/InstanceInfo.h
new file mode 100644
--- /dev/null
+++ b/Engine/Source/Programs/UnrealTraceServer/src/InstanceInfo.h
@@ -0,0 +1,25 @@
+// Copyright Epic Games, Inc. All Rights Reserved.
+
+#pragma once
+
+#include "Version.h"
+
+////////////////////////////////////////////////////////////////////////////////
+struct FInstanceInfo
+{
+public:
+	static const uint32	CurrentVersion =
+#if TS_USING(TS_BUILD_DEBUG)
+		0x8000'0000 |
+#endif
+		((TS_VERSION_PROTOCOL & 0xffff) << 16) | (TS_VERSION_MINOR & 0xffff);
+
+	void				Set();
+	void				WaitForReady() const;
+	bool				IsOlder() const;
+	std::atomic<uint32> Published;
+	uint32				Version;
+	uint32				Pid;
+};
+
+
diff --git a/Engine/Source/Programs/UnrealTraceServer/src/Logging.cpp b/Engine/Source/Programs/UnrealTraceServer/src/Logging.cpp
new file mode 100644
--- /dev/null
+++ b/Engine/Source/Programs/UnrealTraceServer/src/Logging.cpp
@@ -0,0 +1,132 @@
+// Copyright Epic Games, Inc. All Rights Reserved.
+
+#include "Pch.h"
+#include "Logging.h"
+#include "Foundation.h"
+
+////////////////////////////////////////////////////////////////////////////////
+FLogging* FLogging::Instance = nullptr;
+
+////////////////////////////////////////////////////////////////////////////////
+FLogging::FLogging(FPath& InPath)
+	: Path(InPath)
+{
+	// Find where the logs should be written to. Make sure it exists.
+	FPath LogDir = Path.parent_path();
+	FPath Stem = Path.stem();
+
+	// Expected format for log files
+	const FString FilenameFormat = Stem.string() + "_%d.log";
+
+	// Fetch all existing logs.
+	struct FExistingLog
+	{
+		FPath	Path;
+		uint32	Index;
+
+		int32 operator < (const FExistingLog& Rhs) const
+		{
+			return Index < Rhs.Index;
+		}
+	};
+	std::vector<FExistingLog> ExistingLogs;
+	if (std::filesystem::is_directory(LogDir))
+	{
+		
+		for (const auto& DirItem : std::filesystem::directory_iterator(LogDir))
+		{
+			int32 Index = -1;
+			std::string StemUtf8 = DirItem.path().filename().string();
+			sscanf(StemUtf8.c_str(), *FilenameFormat, &Index);
+			if (Index >= 0)
+			{
+				ExistingLogs.push_back({ DirItem.path(), uint32(Index) });
+			}
+		}
+	}
+
+	// Sort and try and tidy up old logs.
+	static int32 MaxLogs = 12; // plus one new one
+	std::sort(ExistingLogs.begin(), ExistingLogs.end());
+	for (int32 i = 0, n = int32(ExistingLogs.size() - MaxLogs); i < n; ++i)
+	{
+		std::error_code ErrorCode;
+		std::filesystem::remove(ExistingLogs[i].Path, ErrorCode);
+	}
+
+
+	// Open the log file (note; can race other instances)
+	uint32 LogIndex = ExistingLogs.empty() ? 0 : ExistingLogs.back().Index;
+	for (uint32 n = LogIndex + 10; File == nullptr && LogIndex < n;)
+	{
+		++LogIndex;
+		char LogName[128];
+		snprintf(LogName, TS_ARRAY_COUNT(LogName), *FilenameFormat, LogIndex);
+		FPath LogPath = LogDir / LogName;
+
+#if TS_USING(TS_PLATFORM_WINDOWS)
+		File = _wfopen(LogPath.c_str(), L"wbxN");
+#else
+		File = fopen(LogPath.c_str(), "wbx");
+#endif
+	}
+}
+
+////////////////////////////////////////////////////////////////////////////////
+FLogging::~FLogging()
+{
+	if (File != nullptr)
+	{
+		fclose(File);
+	}
+}
+
+////////////////////////////////////////////////////////////////////////////////
+void FLogging::LogImpl(const char* String) const
+{
+	if (File != nullptr)
+	{
+		fputs(String, File);
+		fflush(File);
+	}
+
+	fputs(String, stdout);
+#if TS_USING(TS_PLATFORM_WINDOWS) && TS_USING(TS_BUILD_DEBUG)
+	OutputDebugString(String);
+#endif
+}
+
+////////////////////////////////////////////////////////////////////////////////
+void FLogging::Log(const char* Format, ...)
+{
+	va_list VaList;
+	va_start(VaList, Format);
+
+	char Buffer[320];
+	vsnprintf(Buffer, TS_ARRAY_COUNT(Buffer), Format, VaList);
+	Buffer[TS_ARRAY_COUNT(Buffer) - 1] = '\0';
+
+	Instance->LogImpl(Buffer);
+
+	va_end(VaList);
+}
+
+////////////////////////////////////////////////////////////////////////////////
+FLoggingScope::FLoggingScope(FPath& InPath, const char* InFilename /* = nullptr */)
+{
+	FPath Path = InPath / (InFilename ? InFilename : "Server");
+	if (FLogging::Instance != nullptr && FLogging::Instance->Path == Path)
+	{
+		return;
+	}
+
+	PreviousScope = FLogging::Instance;
+	FLogging::Instance = new FLogging(Path);
+}
+
+////////////////////////////////////////////////////////////////////////////////
+FLoggingScope::~FLoggingScope() 
+{
+	delete FLogging::Instance;
+	FLogging::Instance = PreviousScope;
+}
\ No newline at end of file
diff --git a/Engine/Source/Programs/UnrealTraceServer/src/Logging.h b/Engine/Source/Programs/UnrealTraceServer/src/Logging.h
new file mode 100644
--- /dev/null
+++ b/Engine/Source/Programs/UnrealTraceServer/src/Logging.h
@@ -0,0 +1,37 @@
+// Copyright Epic Games, Inc. All Rights Reserved.
+
+#pragma once
+
+#include "Foundation.h"
+
+////////////////////////////////////////////////////////////////////////////////
+#define TS_LOG(Format, ...) \
+	do { FLogging::Log(Format "\n", ##__VA_ARGS__); } while (false)
+
+////////////////////////////////////////////////////////////////////////////////
+class FLogging
+{
+public:
+	static void			Log(const char* Format, ...);
+
+private:
+	friend struct FLoggingScope;
+	FLogging(FPath& Path);
+	~FLogging();
+	FLogging(const FLogging&) = delete;
+	FLogging(FLogging&&) = default;
+	void				LogImpl(const char* String) const;
+	static FLogging* Instance;
+	FPath Path;
+	FILE* File = nullptr;
+};
+
+////////////////////////////////////////////////////////////////////////////////
+struct FLoggingScope
+{
+	FLoggingScope(FPath& Path, const char* LogFileName = nullptr);
+	~FLoggingScope();
+
+private: 
+	FLogging* PreviousScope = nullptr;
+};
\ No newline at end of file
diff --git a/Engine/Source/Programs/UnrealTraceServer/src/Main.cpp b/Engine/Source/Programs/UnrealTraceServer/src/Main.cpp
--- a/Engine/Source/Programs/UnrealTraceServer/src/Main.cpp
+++ b/Engine/Source/Programs/UnrealTraceServer/src/Main.cpp
@@ -5,6 +5,8 @@
 #include "StoreService.h"
 #include "Utils.h"
 #include "Version.h"
+#include "Logging.h"
+#include "InstanceInfo.h"
 
 #include <cxxopts.hpp>
 
@@ -127,166 +129,6 @@ static void GetUnrealTraceHome(FPath& Out, bool Make=false)
 	}
 }
 
-
-
-
-// {{{1 logging ----------------------------------------------------------------
-
-////////////////////////////////////////////////////////////////////////////////
-#define TS_LOG(Format, ...) \
-	do { FLogging::Log(Format "\n", ##__VA_ARGS__); } while (false)
-
-////////////////////////////////////////////////////////////////////////////////
-class FLogging
-{
-public:
-	static void			Initialize();
-	static void			Shutdown();
-	static void			Log(const char* Format, ...);
-
-private:
-						FLogging();
-						~FLogging();
-						FLogging(const FLogging&) = delete;
-						FLogging(FLogging&&) = default;
-	void				LogImpl(const char* String) const;
-	static FLogging*	Instance;
-	FILE*				File = nullptr;
-};
-
-////////////////////////////////////////////////////////////////////////////////
-FLogging* FLogging::Instance = nullptr;
-
-////////////////////////////////////////////////////////////////////////////////
-FLogging::FLogging()
-{
-	// Find where the logs should be written to. Make sure it exists.
-	FPath LogDir;
-	GetUnrealTraceHome(LogDir, true);
-
-	// Fetch all existing logs.
-	struct FExistingLog
-	{
-		FPath	Path;
-		uint32					Index;
-
-		int32 operator < (const FExistingLog& Rhs) const
-		{
-			return Index < Rhs.Index;
-		}
-	};
-	std::vector<FExistingLog> ExistingLogs;
-	if (std::filesystem::is_directory(LogDir))
-	{
-		for (const auto& DirItem : std::filesystem::directory_iterator(LogDir))
-		{
-			int32 Index = -1;
-			std::string StemUtf8 = DirItem.path().stem().string();
-			sscanf(StemUtf8.c_str(), "Server_%d", &Index);
-			if (Index >= 0)
-			{
-				ExistingLogs.push_back({DirItem.path(), uint32(Index)});
-			}
-		}
-	}
-
-	// Sort and try and tidy up old logs.
-	static int32 MaxLogs = 12; // plus one new one
-	std::sort(ExistingLogs.begin(), ExistingLogs.end());
-	for (int32 i = 0, n = int32(ExistingLogs.size() - MaxLogs); i < n; ++i)
-	{
-		std::error_code ErrorCode;
-		std::filesystem::remove(ExistingLogs[i].Path, ErrorCode);
-	}
-
-
-	// Open the log file (note; can race other instances)
-	uint32 LogIndex = ExistingLogs.empty() ? 0 : ExistingLogs.back().Index;
-	for (uint32 n = LogIndex + 10; File == nullptr && LogIndex < n;)
-	{
-		++LogIndex;
-		char LogName[128];
-		snprintf(LogName, TS_ARRAY_COUNT(LogName), "Server_%d.log", LogIndex);
-		FPath LogPath = LogDir / LogName;
-
-#if TS_USING(TS_PLATFORM_WINDOWS)
-		File = _wfopen(LogPath.c_str(), L"wbxN");
-#else
-		File = fopen(LogPath.c_str(), "wbx");
-#endif
-	}
-}
-
-////////////////////////////////////////////////////////////////////////////////
-FLogging::~FLogging()
-{
-	if (File != nullptr)
-	{
-		fclose(File);
-	}
-}
-
-////////////////////////////////////////////////////////////////////////////////
-void FLogging::Initialize()
-{
-	if (Instance != nullptr)
-	{
-		return;
-	}
-
-	Instance = new FLogging();
-}
-
-////////////////////////////////////////////////////////////////////////////////
-void FLogging::Shutdown()
-{
-	if (Instance == nullptr)
-	{
-		return;
-	}
-
-	delete Instance;
-	Instance = nullptr;
-}
-
-////////////////////////////////////////////////////////////////////////////////
-void FLogging::LogImpl(const char* String) const
-{
-	if (File != nullptr)
-	{
-		fputs(String, File);
-		fflush(File);
-	}
-
-	fputs(String, stdout);
-}
-
-////////////////////////////////////////////////////////////////////////////////
-void FLogging::Log(const char* Format, ...)
-{
-	va_list VaList;
-	va_start(VaList, Format);
-
-	char Buffer[320];
-	vsnprintf(Buffer, TS_ARRAY_COUNT(Buffer), Format, VaList);
-	Buffer[TS_ARRAY_COUNT(Buffer) - 1] = '\0';
-
-	Instance->LogImpl(Buffer);
-
-	va_end(VaList);
-}
-
-
-
-////////////////////////////////////////////////////////////////////////////////
-struct FLoggingScope
-{
-	FLoggingScope()		{ FLogging::Initialize(); }
-	~FLoggingScope()	{ FLogging::Shutdown(); }
-};
-
-
-
 // {{{1 store ------------------------------------------------------------------
 
 ////////////////////////////////////////////////////////////////////////////////
@@ -319,69 +161,6 @@ static void ParseOptions(int ArgC, char** ArgV, FStoreSettings* Settings)
 
 }
 
-
-
-// {{{1 instance-info ----------------------------------------------------------
-
-////////////////////////////////////////////////////////////////////////////////
-struct FInstanceInfo
-{
-public:
-	static const uint32	CurrentVersion =
-#if TS_USING(TS_BUILD_DEBUG)
-		0x8000'0000 |
-#endif
-		((TS_VERSION_PROTOCOL & 0xffff) << 16) | (TS_VERSION_MINOR & 0xffff);
-
-	void				Set();
-	void				WaitForReady() const;
-	bool				IsOlder() const;
-	std::atomic<uint32> Published;
-	uint32				Version;
-	uint32				Pid;
-};
-
-////////////////////////////////////////////////////////////////////////////////
-void FInstanceInfo::Set()
-{
-	Version = CurrentVersion;
-#if TS_USING(TS_PLATFORM_WINDOWS)
-	Pid = GetCurrentProcessId();
-#elif TS_USING(TS_PLATFORM_LINUX) || TS_USING(TS_PLATFORM_MAC)
-	Pid = getpid();
-#endif
-	Published.fetch_add(1, std::memory_order_release);
-}
-
-////////////////////////////////////////////////////////////////////////////////
-void FInstanceInfo::WaitForReady() const
-{
-	// Spin until this instance info is published (by another process)
-#if TS_USING(TS_PLATFORM_WINDOWS)
-	for (;; Sleep(0))
-#else
-	for (;; sched_yield())
-#endif
-	{
-		if (Published.load(std::memory_order_acquire))
-		{
-			break;
-		}
-	}
-}
-
-////////////////////////////////////////////////////////////////////////////////
-bool FInstanceInfo::IsOlder() const
-{
-	// Decide which is older; this compiled code or the instance we have a
-	// pointer to.
-	bool bIsOlder = false;
-	bIsOlder |= (Version < FInstanceInfo::CurrentVersion);
-	return bIsOlder;
-}
-
-
-
 // {{{1 return codes -----------------------------------------------------------
 
 ////////////////////////////////////////////////////////////////////////////////
@@ -1334,7 +1113,9 @@ int main(int ArgC, char** ArgV)
 	{
 		if (strcmp(ArgV[1], Dispatch.Verb) == 0)
 		{
-			FLoggingScope LoggingScope;
+			FPath LogDir;
+			GetUnrealTraceHome(LogDir, true);
+			FLoggingScope LoggingScope(LogDir);
 			return (Dispatch.Entry)(ArgC - 1, ArgV + 1);
 		}
 	}
