Fix include issue on Mac/Linux

diff --git a/Engine/Source/Programs/UnrealTraceServer/src/Logging.cpp b/Engine/Source/Programs/UnrealTraceServer/src/Logging.cpp
--- a/Engine/Source/Programs/UnrealTraceServer/src/Logging.cpp
+++ b/Engine/Source/Programs/UnrealTraceServer/src/Logging.cpp
@@ -4,6 +4,10 @@
 #include "Logging.h"
 #include "Foundation.h"
 
+#if TS_USING(TS_PLATFORM_LINUX) || TS_USING(TS_PLATFORM_MAC)
+#include <cstdarg>
+#endif
+
 ////////////////////////////////////////////////////////////////////////////////
 FLogging* FLogging::Instance = nullptr;
 
@@ -129,4 +133,4 @@ FLoggingScope::~FLoggingScope()
 {
 	delete FLogging::Instance;
 	FLogging::Instance = PreviousScope;
-}
\ No newline at end of file
+}
