Implement hooks for store service to shut down if there are no active connections.

diff --git a/Engine/Source/Programs/UnrealTraceServer/src/StoreCborServer.h b/Engine/Source/Programs/UnrealTraceServer/src/StoreCborServer.h
--- a/Engine/Source/Programs/UnrealTraceServer/src/StoreCborServer.h
+++ b/Engine/Source/Programs/UnrealTraceServer/src/StoreCborServer.h
@@ -28,6 +28,7 @@ public:
 	FRecorder&				GetRecorder() const;
 	FTraceRelay*			RelayTrace(uint32 Id);
 	void					OnSettingsChanged();
+	uint32					GetActivePeerCount() const { return Peers.Num(); }
 
 private:
 	virtual bool			OnAccept(asio::ip::tcp::socket& Socket) override;
diff --git a/Engine/Source/Programs/UnrealTraceServer/src/StoreService.cpp b/Engine/Source/Programs/UnrealTraceServer/src/StoreService.cpp
--- a/Engine/Source/Programs/UnrealTraceServer/src/StoreService.cpp
+++ b/Engine/Source/Programs/UnrealTraceServer/src/StoreService.cpp
@@ -83,4 +83,20 @@ uint32 FStoreService::GetRecorderPort() const
 	return Self->Recorder.GetPort();
 }
 
+////////////////////////////////////////////////////////////////////////////////
+bool FStoreService::ShutdownIfNoConnections()
+{
+	auto* Self = (FStoreServiceImpl*)this;
+	const uint32 ConnectionCount = Self->Recorder.GetSessionCount() + Self->CborServer.GetActivePeerCount();
+	if (!ConnectionCount)
+	{
+		// Close the recorder and store server so new connections are not
+		// accepted between returning and process exiting.
+		Self->Recorder.Close();
+		Self->CborServer.Close();
+		return true;
+	}
+	return false;
+}
+
 /* vim: set noexpandtab : */
diff --git a/Engine/Source/Programs/UnrealTraceServer/src/StoreService.h b/Engine/Source/Programs/UnrealTraceServer/src/StoreService.h
--- a/Engine/Source/Programs/UnrealTraceServer/src/StoreService.h
+++ b/Engine/Source/Programs/UnrealTraceServer/src/StoreService.h
@@ -15,6 +15,7 @@ public:
 	void					operator delete (void* Addr);
 	uint32					GetPort() const;
 	uint32					GetRecorderPort() const;
+	bool					ShutdownIfNoConnections();
 
 private:
 							FStoreService() = default;
