Add sponsor mode setting.

Adds sponsor mode setting, and add functionality to print settings.

diff --git a/Engine/Source/Programs/UnrealTraceServer/src/StoreSettings.cpp b/Engine/Source/Programs/UnrealTraceServer/src/StoreSettings.cpp
--- a/Engine/Source/Programs/UnrealTraceServer/src/StoreSettings.cpp
+++ b/Engine/Source/Programs/UnrealTraceServer/src/StoreSettings.cpp
@@ -1,6 +1,7 @@
 // Copyright Epic Games, Inc. All Rights Reserved.
 #include "Pch.h"
 #include "Cbor.h"
+#include "Logging.h"
 #include "StoreSettings.h"
 #include "Utils.h"
 #include <charconv>
@@ -12,13 +13,14 @@ static const char* GAdditionalWatchDirsName = "Additionalwatchdirs";
 static const char* GStorePortName = "StorePort";
 static const char* GRecorderPortName = "RecorderPort";
 static const char* GThreadCountName = "ThreadCount";
+static const char* GSponsoredName = "Sponsored";
 
 ////////////////////////////////////////////////////////////////////////////////
 
 // Logging macro
 #if 1
-#define TS_SETTINGS_LOGF(Fmt, ...) do { printf(Fmt, __VA_ARGS__); }while(false);
-#define TS_SETTINGS_LOG(String) do { printf(String); } while(false);
+#define TS_SETTINGS_LOGF(Fmt, ...) TS_LOG(Fmt, __VA_ARGS__);
+#define TS_SETTINGS_LOG(String) TS_LOG(String);
 #else
 #define TS_SETTINGS_LOGF(...)
 #define TS_SETTINGS_LOG(...)
@@ -275,6 +277,7 @@ void FStoreSettings::ReadFromSettings(const FPath& Path)
 	Reader.GetInteger(GStorePortName, StorePort);
 	Reader.GetInteger(GRecorderPortName, RecorderPort);
 	Reader.GetInteger(GThreadCountName, ThreadCount);
+	Reader.GetInteger(GSponsoredName, Sponsored);
 
 	++ChangeSerial;
 }
@@ -297,6 +300,7 @@ void FStoreSettings::WriteToSettingsFile() const
 	bOk &= Writer.WriteInteger(GStorePortName, StorePort);
 	bOk &= Writer.WriteInteger(GRecorderPortName, RecorderPort);
 	bOk &= Writer.WriteInteger(GThreadCountName, ThreadCount);
+	bOk &= Writer.WriteInteger(GSponsoredName, Sponsored);
 
 	if (!bOk)
 	{
@@ -375,6 +379,11 @@ void FStoreSettings::ApplySettingsFromCbor(const uint8* Buffer, uint32 NumBytes)
 			TS_SETTINGS_TRACE("Set -> %s to %lli\n", GThreadCountName, Context.AsInteger())
 			ThreadCount = int32(Context.AsInteger());
 		}
+		else if (IS_PROPERTY(GSponsoredName, ECborType::Integer))
+		{
+			TS_SETTINGS_TRACE("Set -> %s to %lli\n", GSponsoredName, Context.AsInteger())
+			ThreadCount = int32(Context.AsInteger());
+		}
 	}
 
 	++ChangeSerial;
@@ -413,10 +422,27 @@ void FStoreSettings::SerializeToCbor(TArray<uint8>& OutBuffer) const
 	Writer.WriteString(GThreadCountName);
 	Writer.WriteInteger(ThreadCount);
 
+	Writer.WriteString(GSponsoredName);
+	Writer.WriteInteger(Sponsored);
+
 	Writer.Close();
 
 	OutBuffer.Append(Buffer->GetData(), int32(Buffer->GetSize()));
 }
 
+void FStoreSettings::PrintToLog() const
+{
+	TS_LOG("Store settings (%s):", SettingsFile.string().c_str());
+	TS_LOG(" - Store port: %u", StorePort);
+	TS_LOG(" - Recorder port: %u", RecorderPort);
+	TS_LOG(" - Thread count: %u", ThreadCount);
+	TS_LOG(" - Sponsored mode: %d", Sponsored);
+	TS_LOG(" - Directory: '%s'", StoreDir.string().c_str());
+	for (const auto& Dir : AdditionalWatchDirs)
+	{
+		TS_LOG(" - Additional watch directory: '%s'", Dir.string().c_str());
+	}
+}
+
 #undef PRINTCHECKEDF
 #undef PRINTCHECKED
diff --git a/Engine/Source/Programs/UnrealTraceServer/src/StoreSettings.h b/Engine/Source/Programs/UnrealTraceServer/src/StoreSettings.h
--- a/Engine/Source/Programs/UnrealTraceServer/src/StoreSettings.h
+++ b/Engine/Source/Programs/UnrealTraceServer/src/StoreSettings.h
@@ -14,12 +14,15 @@ public:
 	int32				StorePort		= 1989; // Changes doesn't take effect until restart
 	int32				RecorderPort	= 1981; // Changes doesn't take effect until restart
 	int32				ThreadCount		= 0; 	// <=0:logical CPU count
+	int32				Sponsored		= 1;	// When non-zero store will close when there are
+												// no more sponsor processes or connections active.
 
 	void ReadFromSettings(const FPath& Path);
 	void WriteToSettingsFile() const;
 	void ApplySettingsFromCbor(const uint8* Buffer, uint32 NumBytes);
 	void SerializeToCbor(TArray<uint8>& OutBuffer) const;
 	uint32 GetChangeSerial() const { return ChangeSerial; }
+	void PrintToLog() const;
 
 private:
 	uint32 ChangeSerial = 1;
