Sponsor mode in UnrealTraceServer

Adds new sponsor mode of running server. In this mode the server must be started with a parent pid. New pids can be added by running `unrealtraceserver fork --sponsor [pid]`. When all sponsor processes have exited and there are no active connections the server will exit. Most of the lifetime management logic is located in the new Lifetime type, but actual querying and exiting is done in the respective main functions.

Linux and Mac implementation was using a lock file to communicate running pid and version to other instances. With this change all platforms now use shared memory mapping techniques.

Command line options and command parsing has also been refactored and is now parsed and passed to all main functions. Command line help is also added.

diff --git a/Engine/Source/Programs/UnrealTraceServer/src/InstanceInfo.h b/Engine/Source/Programs/UnrealTraceServer/src/InstanceInfo.h
--- a/Engine/Source/Programs/UnrealTraceServer/src/InstanceInfo.h
+++ b/Engine/Source/Programs/UnrealTraceServer/src/InstanceInfo.h
@@ -8,6 +8,7 @@
 struct FInstanceInfo
 {
 public:
+	static constexpr uint32	NumSponsorQueueSlots = 8;
 	static const uint32	CurrentVersion =
 #if TS_USING(TS_BUILD_DEBUG)
 		0x8000'0000 |
@@ -20,6 +21,9 @@ public:
 	std::atomic<uint32> Published;
 	uint32				Version;
 	uint32				Pid;
+	// Queue slots used to add sponsor processes. The running instance regularly 
+	// polls these slots and moves the pid to an internal representation.
+	std::atomic<uint32>	SponsorPids[NumSponsorQueueSlots];
 };
 
 
diff --git a/Engine/Source/Programs/UnrealTraceServer/src/Lifetime.cpp b/Engine/Source/Programs/UnrealTraceServer/src/Lifetime.cpp
new file mode 100644
--- /dev/null
+++ b/Engine/Source/Programs/UnrealTraceServer/src/Lifetime.cpp
@@ -0,0 +1,103 @@
+// Copyright Epic Games, Inc. All Rights Reserved.
+#include "Pch.h"
+#include "Lifetime.h"
+#include "InstanceInfo.h"
+#include "Logging.h"
+#include "StoreService.h"
+
+#if TS_USING(TS_PLATFORM_LINUX) || TS_USING(TS_PLATFORM_MAC)
+#include <signal.h>
+#endif
+
+////////////////////////////////////////////////////////////////////////////////
+FLifetime::FLifetime(class FStoreService* InStoreService)
+	: StoreService(InStoreService)
+{
+	check(StoreService && InstanceInfo);
+}
+
+////////////////////////////////////////////////////////////////////////////////
+bool FLifetime::ShouldKeepAlive()
+{
+	// Check if all sponsors are alive. If no sponsors are alive try to shutdown
+	// the store if it has no active connections. 
+	const bool bShouldKeepAlive = IsAnySponsorActive() || !ShutdownStoreIfNoConnections();
+	if (!bShouldKeepAlive)
+	{
+		TS_LOG("Terminating server, no sponsors or connections active.");
+	}
+	return bShouldKeepAlive;
+}
+
+////////////////////////////////////////////////////////////////////////////////
+void FLifetime::CheckNewSponsors(FInstanceInfo* InstanceInfo)
+{
+	check(InstanceInfo);
+	for (auto& PidEntry : InstanceInfo->SponsorPids)
+	{
+		if (uint32_t ThisPid = PidEntry.load(std::memory_order_relaxed))
+		{
+			if (PidEntry.compare_exchange_strong(ThisPid, 0))
+			{
+				AddPid(ThisPid);
+			}
+		}
+	}
+}
+
+////////////////////////////////////////////////////////////////////////////////
+void FLifetime::AddPid(uint32 Pid)
+{
+	if (Pid == 0)
+	{
+		return;
+	}
+#if TS_USING(TS_PLATFORM_WINDOWS)
+	FProcHandle ProcessHandle = OpenProcess(PROCESS_QUERY_INFORMATION | SYNCHRONIZE, FALSE, Pid);
+#else
+	FProcHandle ProcessHandle = FProcHandle(intptr_t(Pid));
+#endif
+	SponsorHandles.FindOrAdd(ProcessHandle, [](FProcHandle& Handle, const FProcHandle& New) {return Handle == New ? 0 : 1; });
+}
+
+////////////////////////////////////////////////////////////////////////////////
+bool FLifetime::ShutdownStoreIfNoConnections()
+{
+	return StoreService->ShutdownIfNoConnections();
+}
+
+////////////////////////////////////////////////////////////////////////////////
+bool FLifetime::IsAnySponsorActive()
+{
+	bool bKeepAlive = false;
+	for (auto& Handle : SponsorHandles)
+	{
+#if TS_USING(TS_PLATFORM_WINDOWS)
+		DWORD ExitCode;
+		const bool Result = GetExitCodeProcess(Handle, &ExitCode);
+		if (Result && ExitCode == STILL_ACTIVE)
+		{
+			bKeepAlive = true;
+		}
+		else
+		{
+			Handle = 0;
+		}
+#else
+		int Pid = int(intptr_t(Handle));
+		if (kill(pid_t(Pid), 0) == 0)
+		{
+			bKeepAlive = true;
+		}
+		else
+		{
+			Handle = 0;
+		}
+#endif
+	}
+
+	// Remove inactive processes
+	SponsorHandles.RemoveIf([](const FProcHandle& Handle) { return Handle == 0; });
+	return bKeepAlive;
+}
+
diff --git a/Engine/Source/Programs/UnrealTraceServer/src/Lifetime.h b/Engine/Source/Programs/UnrealTraceServer/src/Lifetime.h
new file mode 100644
--- /dev/null
+++ b/Engine/Source/Programs/UnrealTraceServer/src/Lifetime.h
@@ -0,0 +1,26 @@
+// Copyright Epic Games, Inc. All Rights Reserved.
+
+#pragma once
+
+#include "Foundation.h"
+
+////////////////////////////////////////////////////////////////////////////////
+
+class FLifetime
+{
+public:
+						FLifetime(class FStoreService* Service);
+	bool				ShouldKeepAlive();
+	void				CheckNewSponsors(struct FInstanceInfo* InstanceInfo);
+	void				AddPid(uint32);
+
+private:
+	bool				ShutdownStoreIfNoConnections();
+	bool				IsAnySponsorActive();
+
+	using FProcHandle = void*;
+
+	TArray<FProcHandle> SponsorHandles;
+	class FStoreService* StoreService;
+};
+
diff --git a/Engine/Source/Programs/UnrealTraceServer/src/Main.cpp b/Engine/Source/Programs/UnrealTraceServer/src/Main.cpp
--- a/Engine/Source/Programs/UnrealTraceServer/src/Main.cpp
+++ b/Engine/Source/Programs/UnrealTraceServer/src/Main.cpp
@@ -1,12 +1,13 @@
 // Copyright Epic Games, Inc. All Rights Reserved.
 
 #include "Pch.h"
-#include "StoreSettings.h"
+#include "InstanceInfo.h"
+#include "Lifetime.h"
+#include "Logging.h"
 #include "StoreService.h"
+#include "StoreSettings.h"
 #include "Utils.h"
 #include "Version.h"
-#include "Logging.h"
-#include "InstanceInfo.h"
 
 #include <cxxopts.hpp>
 
@@ -35,6 +36,12 @@
 #	define TS_DAEMON_THREAD TS_OFF
 #endif
 
+// Check for legacy lock files on Linux and Mac 
+#define TS_LEGACY_LOCK_FILE TS_ON
+
+// Frequency between sponsor lifetime checks
+constexpr uint32 GSponsorCheckFreqSecs = 5;
+
 // {{{1 misc -------------------------------------------------------------------
 
 ////////////////////////////////////////////////////////////////////////////////
@@ -131,35 +138,174 @@ static void GetUnrealTraceHome(FPath& Out, bool Make=false)
 
 // {{{1 store ------------------------------------------------------------------
 
-////////////////////////////////////////////////////////////////////////////////
-static void ParseOptions(int ArgC, char** ArgV, FStoreSettings* Settings)
+struct FOptions 
 {
-	cxxopts::Options Options("UnrealTraceServer", "Unreal Trace Server");
-	Options.add_options()
-		("storedir", "Default store directory", cxxopts::value<std::string>()->default_value(""))
-		("port",	"TCP port to serve the store on",			cxxopts::value<int>()->default_value("0"))
-		("recport",	"TCP port for the recorder to listen on",	cxxopts::value<int>()->default_value("0"))
-		;
+public:
+
+	enum ParseResults : int {
+		Result_Continue = 0,
+		Result_MissingCommandError,
+		Result_HelpRequested,
+	};
+
+	FOptions()
+		: Options("UnrealTraceServer", "Unreal Trace Server"
+			"\n\nUnrealTraceServer acts as a hub between runtimes that are tracing performance "
+			"instrumentation and tools like Unreal Insights that consume and present that "
+			"data for analysis. TCP ports 1981 and 1989 are used, where the former receives "
+			"trace data, and the latter is used by tools to query the server's store.\n")
+		, CurrentCommand(nullptr)
+	{
+		// Positional arguments
+		Options.add_options()
+			("command", "Command to execute", cxxopts::value<std::string>()->default_value(""))
+			;
+		Options.parse_positional({ "command" });
+		Options.positional_help("<cmd>");
+
+		Options.add_options()
+			("help", "Prints help message for each command");
+
+		// Fork and daemon options
+		Options.add_options("settings")
+			("storedir", "Default store directory", cxxopts::value<std::string>()->default_value(""))
+			("port", "TCP port to serve the store on", cxxopts::value<int>()->default_value("0"))
+			("recport", "TCP port for the recorder to listen on", cxxopts::value<int>()->default_value("0"))
+			;
+
+		// AddProc options
+		Options.add_options("sponsor")
+			("sponsor", "Pid to add as sponsor. Required if running in sponsored mode.", cxxopts::value<uint32>()->default_value("0"))
+			;
+	}
+
+	ParseResults Parse(int ArgC, char** ArgV)
+	{
+		Parsed = MoveTemp(Options.parse(ArgC, ArgV));
+		const bool bCommandOk = Parsed["command"].count() == 1;
+		const bool bHelp = Parsed["help"].count() > 0;
+
+		// Check for valid command
+		std::string ParsedCommand = Parsed["command"].as<std::string>();
+		for (FCommandHelp* CommandIt = CommandHelp; CommandIt->Command != nullptr; ++CommandIt)
+		{
+			if (strcmp(ParsedCommand.c_str(), CommandIt->Command) == 0)
+			{
+				CurrentCommand = CommandIt;
+				break;
+			}
+		}
+
+		// Update current command help
+		if (!CurrentCommand)
+		{
+			std::stringstream CurrentCommandHelp;
+			CurrentCommandHelp << "<Command>\n\nCommands:\n";
+			for (FCommandHelp* CommandIt = CommandHelp; CommandIt->Command != nullptr; ++CommandIt)
+			{
+				CurrentCommandHelp << "  " << CommandIt->Command << "\t" << CommandIt->ShortText << "\n";
+			}
+			Options.positional_help(CurrentCommandHelp.str());
+		}
+		else
+		{
+			std::stringstream CurrentCommandHelp;
+			CurrentCommandHelp << CurrentCommand->Command << "\n\nCommand:\n  "
+				<< CurrentCommand->ShortText << CurrentCommand->LongText << "\n";
+			Options.positional_help(CurrentCommandHelp.str());
+		}
 
-	auto Parsed = Options.parse(ArgC, ArgV);
+		if (!bCommandOk)
+		{
+			fputs("Error, unknown command specified. Note that command matching is case sensitive.\n", stderr);
+			PrintHelp();
+			return Result_MissingCommandError;
+		}
+		else if (bHelp)
+		{
+			PrintHelp();
+			return Result_HelpRequested;
+		}
+		
+		return Result_Continue;
+	}
 
-	if (int Value = Parsed["port"].as<int>())
+	void ApplyToSettings(FStoreSettings* Settings) const
 	{
-		Settings->StorePort = Value;
+		check(Settings);
+
+		if (int Value = Parsed["port"].as<int>())
+		{
+			Settings->StorePort = Value;
+		}
+
+		if (int Value = Parsed["recport"].as<int>())
+		{
+			Settings->RecorderPort = Value;
+		}
+
+		if (std::string Value = Parsed["storedir"].as<std::string>(); !Value.empty())
+		{
+			Settings->StoreDir = Value;
+		}
 	}
 
-	if (int Value = Parsed["recport"].as<int>())
+	bool GetSponsorPid(uint32& OutSponsorPid) const
 	{
-		Settings->RecorderPort = Value;
+		if (uint32 Value = Parsed["sponsor"].as<uint32>())
+		{
+			OutSponsorPid = Value;
+			return true;
+		}
+		return false;
 	}
 
-	std::string Value = Parsed["storedir"].as<std::string>();
-	if (!Value.empty())
+	void PrintHelp() const
 	{
-		Settings->StoreDir = Value;
+		std::string HelpText;
+		if (!CurrentCommand)
+		{
+			HelpText = Options.help({"dummy"});
+		}
+		else
+		{
+			std::vector<std::string> HelpGroups;
+			for (auto Group : CurrentCommand->HelpGroups)
+			{
+				HelpGroups.push_back(Group);
+			}
+			HelpText = Options.help(HelpGroups);
+		}
+
+		fputs(HelpText.c_str(), stdout);
 	}
 
-}
+
+private:
+	cxxopts::Options Options;
+	cxxopts::ParseResult Parsed;
+	
+	static struct FCommandHelp
+	{
+		const char* Command;
+		const char* ShortText;
+		const char* LongText;
+		std::vector<const char*> HelpGroups;
+	} CommandHelp[];
+
+	const FCommandHelp* CurrentCommand;
+};
+
+FOptions::FCommandHelp FOptions::CommandHelp[] = {
+	"fork", "	Starts a background server, upgrading any existing instance. ", 
+				"Checks if there is an existing instance running. If the running version is the same "
+				"version or newer that instance is used. If a sponsor pid is specified that pid is added "
+				"to the running instance.", {"settings", "sponsor"},
+	"daemon", "The mode that a background server runs in. ", "", {"sponsor"},
+	"kill", "	Shuts down a currently running instance. ", "", {},
+	"test", "	Run tests. ", "", {},
+	nullptr, nullptr, nullptr, {}
+};
 
 // {{{1 return codes -----------------------------------------------------------
 
@@ -183,6 +329,8 @@ enum : int
 	Result_ReadFailPid,
 	Result_LockClaimFail,
 	Result_UnexpectedError,
+	Result_InvalidArgError,
+	Result_SponsorAddFail,
 };
 
 
@@ -354,7 +502,7 @@ static const wchar_t*	GIpcName					= L"Local\\UnrealTraceInstance";
 static const int32		GIpcSize					= 4 << 10;
 static const wchar_t*	GQuitEventName				= L"Local\\UnrealTraceEvent";
 static const wchar_t*	GBegunEventName				= L"Local\\UnrealTraceEventBegun";
-static int				MainDaemon(int, char**);
+static int				MainDaemon(int, char**, const FOptions&);
 void					AddToSystemTray(FStoreService&);
 void					RemoveFromSystemTray();
 
@@ -400,7 +548,7 @@ static int MainKillImpl(int ArgC, char** ArgV, const FInstanceInfo* InstanceInfo
 }
 
 ////////////////////////////////////////////////////////////////////////////////
-static int MainKill(int ArgC, char** ArgV)
+static int MainKill(int ArgC, char** ArgV, const FOptions& Options)
 {
 	// Find if an existing instance is already running.
 	FWinHandle IpcHandle = OpenFileMappingW(FILE_MAP_ALL_ACCESS, false, GIpcName);
@@ -420,7 +568,7 @@ static int MainKill(int ArgC, char** ArgV)
 }
 
 ////////////////////////////////////////////////////////////////////////////////
-static int MainFork(int ArgC, char** ArgV)
+static int MainFork(int ArgC, char** ArgV, const FOptions& Options)
 {
 	// Check for an existing instance that is already running.
 	TS_LOG("Opening exist instance's shared memory");
@@ -430,7 +578,7 @@ static int MainFork(int ArgC, char** ArgV)
 		void* IpcPtr = MapViewOfFile(IpcHandle, FILE_MAP_ALL_ACCESS, 0, 0, GIpcSize);
 		FMmapScope MmapScope(IpcPtr);
 
-		const auto* InstanceInfo = MmapScope.As<const FInstanceInfo>();
+		auto* InstanceInfo = MmapScope.As<FInstanceInfo>();
 		InstanceInfo->WaitForReady();
 #if TS_USING(TS_BUILD_DEBUG)
 		if (false)
@@ -439,6 +587,28 @@ static int MainFork(int ArgC, char** ArgV)
 #endif
 		{
 			TS_LOG("Existing instance is the same age or newer");
+
+			// If a parent pid was specified, add to the list of sponsor processes.
+			if (uint32 PidToAdd = 0; Options.GetSponsorPid(PidToAdd))
+			{
+				bool bSuccess = false;
+				for (auto& Pid : InstanceInfo->SponsorPids)
+				{
+					uint32 Expected = 0;
+					if (Pid.compare_exchange_strong(Expected, PidToAdd))
+					{
+						bSuccess = true;
+						break;
+					}
+				}
+
+				if (!bSuccess)
+				{
+					TS_LOG("Failed to add sponsor process in existing instance.");
+					return Result_SponsorAddFail;
+				}
+			}
+
 			return Result_Ok;
 		}
 
@@ -541,14 +711,29 @@ static int MainFork(int ArgC, char** ArgV)
 	// For debugging ease and consistency we will daemonize in this process
 	// instead of spawning a second one.
 #if TS_USING(TS_DAEMON_THREAD)
-	std::thread DaemonThread([=] () { MainDaemon(ArgC, ArgV); });
+	std::thread DaemonThread([=] () { MainDaemon(ArgC, ArgV, Options); });
 #else
 	std::wstring CommandLine = L"UnrealTraceServer.exe daemon";
+
+	auto ContainsSpace = [](LPCWSTR Arg) -> bool {
+		for (uint32 c = 0, Len = (uint32)wcslen(Arg); c < Len; ++c)
+		{
+			if (iswspace(Arg[c]))
+			{
+				return true;
+			}
+		}
+		return false;
+	};
+
 	for (int i = 1; i < ArgC; ++i)
 	{
-		CommandLine += L"\"";
+		FWinApiStr Arg(ArgV[i]);
+		const bool bContainsSpace = ContainsSpace(Arg);
+		CommandLine += L" ";
+		if (bContainsSpace) CommandLine += L"\"";
 		CommandLine += FWinApiStr(ArgV[i]);
-		CommandLine += L"\"";
+		if (bContainsSpace) CommandLine += L"\"";
 	}
 	if (!LaunchUnelevated(DestPath.c_str(), CommandLine.data()))
 	{
@@ -595,7 +780,7 @@ static int MainFork(int ArgC, char** ArgV)
 }
 
 ////////////////////////////////////////////////////////////////////////////////
-static int MainDaemon(int ArgC, char** ArgV)
+static int MainDaemon(int ArgC, char** ArgV, const FOptions& Options)
 {
 	// Move the working directory to be where this binary is located.
 	TS_LOG("Setting working directory");
@@ -638,45 +823,39 @@ static int MainDaemon(int ArgC, char** ArgV)
 
 	// Fill out the Ipc details and publish.
 	void* IpcPtr = MapViewOfFile(IpcHandle, FILE_MAP_ALL_ACCESS, 0, 0, GIpcSize);
-	{
-		TS_LOG("Writing shared instance info");
-		FMmapScope MmapScope(IpcPtr);
-		auto* InstanceInfo = MmapScope.As<FInstanceInfo>();
-		InstanceInfo->Set();
-	}
+	
+	TS_LOG("Writing shared instance info");
+	FMmapScope MmapScope(IpcPtr);
+	auto* InstanceInfo = MmapScope.As<FInstanceInfo>();
+	InstanceInfo->Set();
+	
 
 	// Set locale for std function application wide.
 	std::setlocale(LC_ALL, ".UTF-8");
+	
+	TS_LOG("Starting the store");
+	FPath HomeDir;
+	GetUnrealTraceHome(HomeDir);
 
-	// Fire up the store
+	// Read settings from configuration file
 	FStoreSettings* Settings = new FStoreSettings();
-	FStoreService* StoreService = nullptr;
+	Settings->ReadFromSettings(HomeDir);
+	// Override with command line arguments.
+	Options.ApplyToSettings(Settings);
+	// Display final settings to user
+	Settings->PrintToLog();
+
+	// Read other arguments
+	uint32 ParentPid = 0;
+	if (Settings->Sponsored && !Options.GetSponsorPid(ParentPid))
 	{
-		 TS_LOG("Starting the store");
-		 FPath HomeDir;
-		 GetUnrealTraceHome(HomeDir);
-
-		 // Read settings from configuration file
-		 Settings->ReadFromSettings(HomeDir);
-
-		 // Override with command line arguments.
-		 ParseOptions(ArgC, ArgV, Settings);
-
-		 TS_LOG("Creating store with settings:");
-		 TS_LOG(" - Directory: '%s'", Settings->StoreDir.string().c_str());
-		 TS_LOG(" - Store port: %u", Settings->StorePort);
-		 TS_LOG(" - Recorder port: %u", Settings->RecorderPort);
-		 if (Settings->AdditionalWatchDirs.Num())
-		 {
-			 for (const auto& Dir : Settings->AdditionalWatchDirs)
-			 {
-				 TS_LOG(" - Additional watch directory: '%s'", Dir.string().c_str());
-			 }
-		 }
-
-		 StoreService = FStoreService::Create(Settings);
+		TS_LOG("Error: Deamon is configured to run in sponsored mode, but no sponsor pid has been specified.");
+		return Result_InvalidArgError;
 	}
 
+	// Fire up the store
+	FStoreService* StoreService = FStoreService::Create(Settings);
+
 	// Let every one know we've started.
 	{
 		FWinHandle BegunEvent = CreateEventW(nullptr, TRUE, FALSE, GBegunEventName);
@@ -691,7 +870,22 @@ static int MainDaemon(int ArgC, char** ArgV)
 	AddToSystemTray(*StoreService);
 
 	// Wait to be told to resign.
-	WaitForSingleObject(QuitEvent, INFINITE);
+	FLifetime LifetimeManager(StoreService);
+	LifetimeManager.AddPid(ParentPid);
+
+	while (true)
+	{
+		DWORD Result = WaitForSingleObject(QuitEvent, GSponsorCheckFreqSecs * 1000);
+		if (Result != WAIT_TIMEOUT)
+		{
+			break;
+		}
+		LifetimeManager.CheckNewSponsors(InstanceInfo);
+		if (Settings->Sponsored && !LifetimeManager.ShouldKeepAlive())
+		{
+			break;
+		}
+	}
 
 	// Clean up. We are done here.
 	RemoveFromSystemTray();
@@ -738,14 +932,59 @@ int WINAPI WinMain(HINSTANCE, HINSTANCE, LPSTR, int)
 #if TS_USING(TS_PLATFORM_LINUX) || TS_USING(TS_PLATFORM_MAC)
 
 ////////////////////////////////////////////////////////////////////////////////
-static int MainDaemon(int, char**, pid_t);
+static int MainDaemonImpl(int, char**, pid_t, const FOptions& Options);
 
 ////////////////////////////////////////////////////////////////////////////////
+#if TS_USING(TS_LEGACY_LOCK_FILE)
+
+static int MainKillImpl(int ArgC, char** ArgV, pid_t DaemonPid);
+
 static FPath GetLockFilePath()
 {
 	return "/tmp/UnrealTraceServer.pid";
 }
 
+static int LegacyLockFile()
+{
+	// Open the pid file to detect an existing instance
+	FPath DotPidPath = GetLockFilePath();
+	TS_LOG("Checking for a '%s' lock file", DotPidPath.c_str());
+	for (int DotPidFd = open(DotPidPath.c_str(), O_RDONLY); DotPidFd >= 0; )
+	{
+		OnScopeExit([DotPidFd] () { close(DotPidFd); });
+
+		// If we can claim a write lock then the lock file is orphaned
+		struct flock FileLock = { .l_type = F_WRLCK };
+		int Result = fcntl(DotPidFd, F_GETLK, &FileLock);
+		if (Result != 0 || FileLock.l_type == F_UNLCK)
+		{
+			TS_LOG("Lock file appears to be orphaned. Unlinking it");
+			std::error_code ErrorCode;
+			std::filesystem::remove(DotPidPath, ErrorCode);
+			break;
+		}
+
+		// An instance running with pid file is by definition older than us
+		TS_LOG("Killing an older instance that is already running");
+		int KillRet = MainKillImpl(0, nullptr, FileLock.l_pid);
+		if (KillRet == Result_NoQuitEvent)
+		{
+			// If no quit event was found then we shall assume that another new
+			// store instance beat us to it.
+			TS_LOG("Looks like someone else has already taken care of the upgrade");
+			return Result_Ok;
+		}
+
+		if (KillRet != Result_Ok)
+		{
+			TS_LOG("Kill attempt failed (ret=%d)", KillRet);
+			return KillRet;
+		}
+	}
+	return Result_Ok;
+}
+#endif
+
 ////////////////////////////////////////////////////////////////////////////////
 static int MainKillImpl(int ArgC, char** ArgV, pid_t DaemonPid)
 {
@@ -787,80 +1026,99 @@ static int MainKillImpl(int ArgC, char** ArgV, pid_t DaemonPid)
 }
 
 ////////////////////////////////////////////////////////////////////////////////
-static int MainKill(int ArgC, char** ArgV)
+static int MainKill(int ArgC, char** ArgV, const FOptions& Options)
 {
-	// Open the pid file to detect an existing instance
-	FPath DotPidPath = GetLockFilePath();
-	TS_LOG("Checking for a '%s' lock file", DotPidPath.c_str());
-	int DotPidFd = open(DotPidPath.c_str(), O_RDONLY);
-	if (DotPidFd < 0)
+	TS_LOG("Opening shared memory");
+	int Fd = shm_open("/UnrealTraceServer", O_RDONLY | O_CLOEXEC, 0666);
+	if (Fd > 0)
 	{
-		if (errno == ENOENT)
+		void* BufferPtr = mmap(nullptr, sizeof(FInstanceInfo), PROT_READ, MAP_SHARED, Fd, 0);
+		if (!BufferPtr)
+		{
+			TS_LOG("Unable to map shared memory");
+			return Result_SharedMemFail;
+		}
+
+		FMmapScope Buffer(BufferPtr);
+		FInstanceInfo* InstanceInfo = Buffer.As<FInstanceInfo>();
+		
+		// If we've got this far then there's an instance running that is old
+		TS_LOG("Killing an older instance that is already running");
+		int KillRet = MainKillImpl(0, nullptr, InstanceInfo->Pid);
+		if (KillRet == Result_NoQuitEvent)
 		{
-			TS_LOG("All good. Ain't nuffin' running me ol' pal.");
+			// If no quit event was found then we shall assume that another new
+			// store instance beat us to it.
+			TS_LOG("Looks like someone else has already taken care of the upgrade");
 			return Result_Ok;
 		}
 
-		TS_LOG("Unable to open lock file (%s, errno=%d)", DotPidPath.c_str(), errno);
-		return Result_OpenFailPid;
+		if (KillRet != Result_Ok)
+		{
+			TS_LOG("Kill attempt failed (ret=%d)", KillRet);
+			return KillRet;
+		}
 	}
-	OnScopeExit([DotPidFd] { close(DotPidFd); });
-
-	// If we can claim the write lock then lock file is orphaned.
-	struct flock FileLock = { .l_type = F_WRLCK };
-	int Result = fcntl(DotPidFd, F_GETLK, &FileLock);
-	if (Result != 0 || FileLock.l_type == F_UNLCK)
+#if TS_USING(TS_LEGACY_LOCK_FILE)
+	else
 	{
-		TS_LOG("Lock file appears to be orphaned. Removing");
-		std::error_code ErrorCode;
-		std::filesystem::remove(DotPidPath, ErrorCode);
-		return Result_Ok;
+		TS_LOG("Shared memory doesn't exist, checking legacy lock file");
+		return LegacyLockFile(); 
 	}
+#endif
 
-	return MainKillImpl(ArgC, ArgV, FileLock.l_pid);
+	return Result_Ok;
 }
 
+
 ////////////////////////////////////////////////////////////////////////////////
-static int MainFork(int ArgC, char** ArgV)
+static int MainFork(int ArgC, char** ArgV, const FOptions& Options)
 {
-	// Open the pid file to detect an existing instance
-	FPath DotPidPath = GetLockFilePath();
-	TS_LOG("Checking for a '%s' lock file", DotPidPath.c_str());
-	for (int DotPidFd = open(DotPidPath.c_str(), O_RDONLY); DotPidFd >= 0; )
+	TS_LOG("Opening shared memory");
+	int Fd = shm_open("/UnrealTraceServer", O_RDWR | O_CLOEXEC, 0666);
+	if (Fd > 0)
 	{
-		OnScopeExit([DotPidFd] () { close(DotPidFd); });
-
-		// If we can claim a write lock then the lock file is orphaned
-		struct flock FileLock = { .l_type = F_WRLCK };
-		int Result = fcntl(DotPidFd, F_GETLK, &FileLock);
-		if (Result != 0 || FileLock.l_type == F_UNLCK)
+		void* BufferPtr = mmap(nullptr, sizeof(FInstanceInfo), PROT_READ | PROT_WRITE, MAP_SHARED, Fd, 0);
+		if (!BufferPtr)
 		{
-			TS_LOG("Lock file appears to be orphaned. Unlinking it");
-			std::error_code ErrorCode;
-			std::filesystem::remove(DotPidPath, ErrorCode);
-			break;
+			TS_LOG("Unable to map shared memory");
+			return Result_SharedMemFail;
 		}
 
-		// Get the instance info from the buffer
-		char DotPidBuffer[sizeof(FInstanceInfo)];
-		Result = read(DotPidFd, DotPidBuffer, sizeof(DotPidBuffer));
-		if (Result < sizeof(FInstanceInfo))
-		{
-			TS_LOG("Failed to read the .pid lock file (errno=%d)", errno);
-			return Result_ReadFailPid;
-		}
-		const auto* InstanceInfo = (const FInstanceInfo*)DotPidBuffer;
+		FMmapScope Buffer(BufferPtr);
+		FInstanceInfo* InstanceInfo = Buffer.As<FInstanceInfo>();
 
 		// Old enough for this fine establishment?
 		if (!InstanceInfo->IsOlder())
 		{
 			TS_LOG("Existing instance is the same age or newer");
+
+			if (uint32 PidToAdd = 0; Options.GetSponsorPid(PidToAdd))
+			{
+				bool bSuccess = false;
+				for (auto& Pid : InstanceInfo->SponsorPids)
+				{
+					uint32 Expected = 0;
+					if (Pid.compare_exchange_strong(Expected, PidToAdd))
+					{
+						bSuccess = true;
+						break;
+					}
+				}
+
+				if (!bSuccess)
+				{
+					TS_LOG("Failed to add sponsor process in existing instance.");
+					return Result_SponsorAddFail;
+				}
+			}
+
 			return Result_Ok;
 		}
 
 		// If we've got this far then there's an instance running that is old
 		TS_LOG("Killing an older instance that is already running");
-		int KillRet = MainKillImpl(0, nullptr, FileLock.l_pid);
+		int KillRet = MainKillImpl(0, nullptr, InstanceInfo->Pid);
 		if (KillRet == Result_NoQuitEvent)
 		{
 			// If no quit event was found then we shall assume that another new
@@ -875,13 +1133,23 @@ static int MainFork(int ArgC, char** ArgV)
 			return KillRet;
 		}
 	}
+#if TS_USING(TS_LEGACY_LOCK_FILE)
+	else
+	{
+		TS_LOG("Shared memory doesn't exist, checking legacy lock file");
+		if (int Result = LegacyLockFile(); Result != Result_Ok)
+		{
+			return Result;
+		}
+	}
+#endif
 
 	// Launch a daemonized version of ourselves. For debugging ease and
 	// consistency we will daemonize in this process instead of spawning
 	// a second one.
 	pid_t DaemonPid = -1;
-#if TS_USING(TS_BUILD_DEBUG)
-	std::thread DaemonThread([=] () { MainDaemon(ArgC, ArgV, 0); });
+#if TS_USING(TS_DAEMON_THREAD)
+	std::thread DaemonThread([=] () { MainDaemonImpl(ArgC, ArgV, 0, Options); });
 #else
 	TS_LOG("Forking process");
 	pid_t ParentPid = getpid();
@@ -893,13 +1161,13 @@ static int MainFork(int ArgC, char** ArgV)
 	}
 	else if (DaemonPid == 0)
 	{
-		return MainDaemon(ArgC, ArgV, ParentPid);
+		return MainDaemonImpl(ArgC, ArgV, ParentPid, Options);
 	}
 #endif // TS_BUILD_DEBUG
 
 	// Wait for the daemon to indicate that it has started the store.
 	int Ret = Result_Ok;
-#if TS_USING(TS_BUILD_DEBUG)
+#if TS_USING(TS_DAEMON_THREAD)
 	DaemonThread.join();
 #else
 	TS_LOG("Wait until we know the daemon has started.");
@@ -945,59 +1213,64 @@ static int MainFork(int ArgC, char** ArgV)
 }
 
 ////////////////////////////////////////////////////////////////////////////////
-static int MainDaemon(int ArgC, char** ArgV, pid_t ParentPid)
+static int MainDaemonImpl(int ArgC, char** ArgV, pid_t ParentPid, const FOptions& Options)
 {
-	// We expect that there is no lock file on disk if we've got this far.
-	FPath DotPidPath = GetLockFilePath();
-	TS_LOG("Opening the lock file '%s'", DotPidPath.c_str());
-	int DotPidFd = open(DotPidPath.c_str(), O_CREAT|O_WRONLY, 0666);
-	if (DotPidFd < 0)
+	TS_LOG("Opening shared memory");
+	int Fd = shm_open("/UnrealTraceServer", O_RDWR | O_CREAT | O_CLOEXEC, 0666);
+	if (Fd < 0)
 	{
-		TS_LOG("Failed to open lock file (errno=%d)", errno);
-		return Result_OpenFailPid;
+		TS_LOG("Unable to create shared memory");
+		return Result_SharedMemFail;
 	}
-	fchmod(DotPidFd, 0666);
-	OnScopeExit([DotPidFd] () { close(DotPidFd); });
 
-	// Lock the lock file
-	struct flock FileLock = { .l_type = F_WRLCK };
-	int Result = fcntl(DotPidFd, F_SETLK, &FileLock);
-	if (Result < 0)
+	// Set the size of the file
+	if (ftruncate(Fd, sizeof(FInstanceInfo)))
 	{
-		TS_LOG("Unable to claim lock file (errno=%d)", errno);
-		return Result_LockClaimFail;
+		return Result_SharedMemFail;
 	}
 
+	void* BufferPtr = mmap(nullptr, sizeof(FInstanceInfo), PROT_READ | PROT_WRITE, MAP_SHARED, Fd, 0);
+	if (!BufferPtr)
+	{
+		TS_LOG("Unable to map shared memory");
+		return Result_SharedMemFail;
+	}
+
+	FMmapScope Buffer(BufferPtr);
+	FInstanceInfo* InstanceInfo = Buffer.As<FInstanceInfo>();
+
 	// Fire up the store
 	TS_LOG("Starting the store");
 	FStoreSettings* Settings = new FStoreSettings();
 	FStoreService* StoreService = nullptr;
-	{
-		FPath HomeDir;
-		GetUnrealTraceHome(HomeDir);
+	
+	FPath HomeDir;
+	GetUnrealTraceHome(HomeDir);
 
-		// Read settings from configuration file
-		Settings->ReadFromSettings(HomeDir);
+	// Read settings from configuration file
+	Settings->ReadFromSettings(HomeDir);
 
-		// Override with command line arguments
-		ParseOptions(ArgC, ArgV, Settings);
+	// Override with command line arguments
+	Options.ApplyToSettings(Settings);
+	// Display final settings to user
+	Settings->PrintToLog();
 
-		StoreService = FStoreService::Create(Settings);
+	// Get sponsor pid
+	uint32 SponsorPid = 0;
+	if (Settings->Sponsored && !Options.GetSponsorPid(SponsorPid))
+	{
+		TS_LOG("Error: Deamon is configured to run in sponsored mode, but no sponsor pid has been specified.");
+		return Result_InvalidArgError;
 	}
+
+	StoreService = FStoreService::Create(Settings);
 	OnScopeExit([StoreService] () { delete StoreService; });
 
-	// Fill out the lock file with details about this instance
-	FInstanceInfo InstanceInfo;
-	InstanceInfo.Set();
-	if (write(DotPidFd, &InstanceInfo, sizeof(InstanceInfo)) != sizeof(InstanceInfo))
-	{
-		TS_LOG("Unable to write instance info to lock file (errno=%d)", errno);
-		return Result_UnexpectedError;
-	}
-	fsync(DotPidFd);
+	// Fill out the shared memory with details about this instance
+	InstanceInfo->Set();
 
 	// Detach from our parent and from any controlling terminal
-	for (Result = setsid(); Result < 0;)
+	for (int Result = setsid(); Result < 0;)
 	{
 		if (errno == EPERM)
 		{
@@ -1028,31 +1301,39 @@ static int MainDaemon(int ArgC, char** ArgV, pid_t ParentPid)
 	sigaddset(&SignalSet, SIGTERM);
 	sigaddset(&SignalSet, SIGKILL);
 	sigaddset(&SignalSet, SIGINT);
+
+	timespec Timeout;
+	Timeout.tv_sec = GSponsorCheckFreqSecs;
+	FLifetime LifetimeManager(StoreService);
+	LifetimeManager.AddPid(SponsorPid);
+
 	while (true)
 	{
-		int Signal = -1;
-		int Ret = sigwait(&SignalSet, &Signal);
-		if (Ret == 0)
+		siginfo_t SignalInfo;
+		int Ret = sigtimedwait(&SignalSet, &SignalInfo, &Timeout);
+		if (Ret == EINTR)
+		{
+			TS_LOG("Received signal %d", SignalInfo.si_signo);
+			break;
+		}
+
+		LifetimeManager.CheckNewSponsors(InstanceInfo);
+		if (Settings->Sponsored && !LifetimeManager.ShouldKeepAlive())
 		{
-			TS_LOG("Received signal %d", Signal);
 			break;
 		}
 	}
 
 	// Clean up. We are done here.
-	std::error_code ErrorCode;
-	std::filesystem::remove(DotPidPath, ErrorCode);
-
-	FileLock.l_type = F_UNLCK;
-	fcntl(DotPidFd, F_SETLK, &FileLock);
+	shm_unlink("/UnrealTraceServer");
 
 	return Result_Ok;
 }
 
 ////////////////////////////////////////////////////////////////////////////////
-static int MainDaemon(int ArgC, char** ArgV)
+static int MainDaemon(int ArgC, char** ArgV, const FOptions& Options)
 {
-	return MainDaemon(ArgC, ArgV, 0);
+	return MainDaemonImpl(ArgC, ArgV, 0, Options);
 }
 
 #endif // TS_PLATFORM_LINUX/MAC
@@ -1062,7 +1343,7 @@ static int MainDaemon(int ArgC, char** ArgV)
 // {{{1 main -------------------------------------------------------------------
 
 ////////////////////////////////////////////////////////////////////////////////
-int MainTest(int ArgC, char** ArgV)
+int MainTest(int ArgC, char** ArgV, const FOptions& Options)
 {
 	extern void TestCbor();
 	TestCbor();
@@ -1072,36 +1353,15 @@ int MainTest(int ArgC, char** ArgV)
 ////////////////////////////////////////////////////////////////////////////////
 int main(int ArgC, char** ArgV)
 {
-	if (ArgC < 2)
+	FOptions Options;
+	if (auto ExitReason = Options.Parse(ArgC, ArgV); ExitReason)
 	{
-		printf("UnrealTraceServer v%d.%d / Unreal Engine / Epic Games\n\n", TS_VERSION_PROTOCOL, TS_VERSION_MINOR);
-		puts(  "Usage; <cmd>");
-		puts(  "Commands;");
-		puts(  "  fork   Starts a background server, upgrading any existing instance");
-		puts(  "  daemon The mode that a background server runs in");
-		puts(  "  kill   Shuts down a currently running instance");
-		puts("");
-		puts(  "UnrealTraceServer acts as a hub between runtimes that are tracing performance");
-		puts(  "instrumentation and tools like Unreal Insights that consume and present that");
-		puts(  "data for analysis. TCP ports 1981 and 1989 are used, where the former receives");
-		puts(  "trace data, and the latter is used by tools to query the server's store.");
-
-		FPath HomeDir;
-		GetUnrealTraceHome(HomeDir);
-		HomeDir.make_preferred();
-		std::string HomeDirU8 = HomeDir.string();
-		printf("\nStore path; %s\n", HomeDirU8.c_str());
-
-		return 127;
+		return ExitReason == FOptions::Result_HelpRequested ? Result_Ok : Result_InvalidArgError;
 	}
 
-#ifdef TS_WITH_MUSL
-	printf("Using musl\n");
-#endif
-
 	struct {
 		const char*	Verb;
-		int			(*Entry)(int, char**);
+		int			(*Entry)(int, char**, const FOptions&);
 	} Dispatches[] = {
 		"fork",		MainFork,
 		"daemon",	MainDaemon,
@@ -1116,12 +1376,16 @@ int main(int ArgC, char** ArgV)
 			FPath LogDir;
 			GetUnrealTraceHome(LogDir, true);
 			FLoggingScope LoggingScope(LogDir);
-			return (Dispatch.Entry)(ArgC - 1, ArgV + 1);
+
+#ifdef TS_WITH_MUSL
+			TS_LOG("Using musl");
+#endif
+			return (Dispatch.Entry)(ArgC - 1, ArgV + 1, Options);
 		}
 	}
 
 	printf("Unknown command '%s'\n", ArgV[1]);
-	return 126;
+	return Result_InvalidArgError;
 }
 
 /* vim: set noexpandtab foldlevel=1 : */
