Handle edge cases for many simultanious sponsors

In extreme cases (for example user starting server + multiple local clients) all tracing at the exact time the number of command slots might be too small. Bump the number of sponsor command slots to 128 and add a retry in the process that adds commands.

diff --git a/Engine/Source/Programs/UnrealTraceServer/src/InstanceInfo.h b/Engine/Source/Programs/UnrealTraceServer/src/InstanceInfo.h
--- a/Engine/Source/Programs/UnrealTraceServer/src/InstanceInfo.h
+++ b/Engine/Source/Programs/UnrealTraceServer/src/InstanceInfo.h
@@ -8,7 +8,7 @@
 struct FInstanceInfo
 {
 public:
-	static constexpr uint32	NumSponsorQueueSlots = 8;
+	static constexpr uint32	NumSponsorQueueSlots = 128;
 	static const uint32	CurrentVersion =
 #if TS_USING(TS_BUILD_DEBUG)
 		0x8000'0000 |
diff --git a/Engine/Source/Programs/UnrealTraceServer/src/Main.cpp b/Engine/Source/Programs/UnrealTraceServer/src/Main.cpp
--- a/Engine/Source/Programs/UnrealTraceServer/src/Main.cpp
+++ b/Engine/Source/Programs/UnrealTraceServer/src/Main.cpp
@@ -592,13 +592,21 @@ static int MainFork(int ArgC, char** ArgV, const FOptions& Options)
 			if (uint32 PidToAdd = 0; Options.GetSponsorPid(PidToAdd))
 			{
 				bool bSuccess = false;
-				for (auto& Pid : InstanceInfo->SponsorPids)
+				while (!bSuccess)
 				{
-					uint32 Expected = 0;
-					if (Pid.compare_exchange_strong(Expected, PidToAdd))
+					for (auto& Pid : InstanceInfo->SponsorPids)
 					{
-						bSuccess = true;
-						break;
+						uint32 Expected = 0;
+						if (Pid.compare_exchange_strong(Expected, PidToAdd))
+						{
+							bSuccess = true;
+							break;
+						}
+					}
+					if (!bSuccess)
+					{
+						TS_LOG("Sponsor slots full, relax a second...");
+						Sleep(1000);
 					}
 				}
 
@@ -881,7 +889,7 @@ static int MainDaemon(int ArgC, char** ArgV, const FOptions& Options)
 			break;
 		}
 		LifetimeManager.CheckNewSponsors(InstanceInfo);
-		if (Settings->Sponsored && !LifetimeManager.ShouldKeepAlive())
+		if (!LifetimeManager.ShouldKeepAlive() && Settings->Sponsored)
 		{
 			break;
 		}
@@ -1096,13 +1104,21 @@ static int MainFork(int ArgC, char** ArgV, const FOptions& Options)
 			if (uint32 PidToAdd = 0; Options.GetSponsorPid(PidToAdd))
 			{
 				bool bSuccess = false;
-				for (auto& Pid : InstanceInfo->SponsorPids)
+				while (!bSuccess)
 				{
-					uint32 Expected = 0;
-					if (Pid.compare_exchange_strong(Expected, PidToAdd))
+					for (auto& Pid : InstanceInfo->SponsorPids)
+					{
+						uint32 Expected = 0;
+						if (Pid.compare_exchange_strong(Expected, PidToAdd))
+						{
+							bSuccess = true;
+							break;
+						}
+					}
+					if (!bSuccess)
 					{
-						bSuccess = true;
-						break;
+						TS_LOG("Sponsor slots full, relax a second...");
+						sleep(1);
 					}
 				}
 
@@ -1318,7 +1334,7 @@ static int MainDaemonImpl(int ArgC, char** ArgV, pid_t ParentPid, const FOptions
 		}
 
 		LifetimeManager.CheckNewSponsors(InstanceInfo);
-		if (Settings->Sponsored && !LifetimeManager.ShouldKeepAlive())
+		if (!LifetimeManager.ShouldKeepAlive() && Settings->Sponsored)
 		{
 			break;
 		}
