Write to correct settings name

diff --git a/Engine/Source/Programs/UnrealTraceServer/src/StoreSettings.cpp b/Engine/Source/Programs/UnrealTraceServer/src/StoreSettings.cpp
--- a/Engine/Source/Programs/UnrealTraceServer/src/StoreSettings.cpp
+++ b/Engine/Source/Programs/UnrealTraceServer/src/StoreSettings.cpp
@@ -382,7 +382,7 @@ void FStoreSettings::ApplySettingsFromCbor(const uint8* Buffer, uint32 NumBytes)
 		else if (IS_PROPERTY(GSponsoredName, ECborType::Integer))
 		{
 			TS_SETTINGS_TRACE("Set -> %s to %lli\n", GSponsoredName, Context.AsInteger())
-			ThreadCount = int32(Context.AsInteger());
+			Sponsored = int32(Context.AsInteger());
 		}
 	}
 
