Add sponsor and port information in status message

diff --git a/Engine/Source/Programs/UnrealTraceServer/src/StoreCborServer.cpp b/Engine/Source/Programs/UnrealTraceServer/src/StoreCborServer.cpp
--- a/Engine/Source/Programs/UnrealTraceServer/src/StoreCborServer.cpp
+++ b/Engine/Source/Programs/UnrealTraceServer/src/StoreCborServer.cpp
@@ -155,6 +155,8 @@ void FStoreCborPeer::OnStatus()
 	const FStoreSettings* Settings = Parent.GetSettings();
 	TPayloadBuilder<> Builder(EStatusCode::Success);
 	Builder.AddInteger("recorder_port", Settings->RecorderPort);
+	Builder.AddInteger("store_port", Settings->StorePort);
+	Builder.AddInteger("sponsored", Settings->Sponsored);
 	Builder.AddInteger("change_serial", Store.GetChangeSerial());
 	Builder.AddInteger("settings_serial", Settings->GetChangeSerial());
 	Builder.AddString("store_dir", Settings->StoreDir.string().c_str());
