Add support for universal binaries for UTS.

diff --git a/Engine/Source/Programs/UnrealTraceServer/Readme.md b/Engine/Source/Programs/UnrealTraceServer/Readme.md
--- a/Engine/Source/Programs/UnrealTraceServer/Readme.md
+++ b/Engine/Source/Programs/UnrealTraceServer/Readme.md
@@ -60,3 +60,13 @@ Optionally IDE files can be generated, for example Visual Studio solution, by us
 
 The install script will check out the binaries deployed in the Unreal Engine tree. Submit the binaries together with the `Version.h`.
 
+## Building a universal binary on MacOS
+
+In order to build a universal binary for MacOS you need to first build the `x86_64` binary and then the `arm64` binary. These can be combined during
+the installation script.
+
+```
+> xmake config -a x86_64; xmake
+> xmake config -a arm64; xmake 
+> xmake install
+```
