Fix POSIX incompatibilities in Mac implementation

* The sig_timedwait function was used to listen for signals and regularly check sponsor processes. However this function is not available on MacOS. Changed the daemon function to use the old signal listening method but add a worker thread that checks the sponsors at regular intervals instead. Send SIGUSR2 when exit is desired.
* Moved log output from Lifetime to Main to avoid spamming when sponsored mode is disabled. Note that we still need to call ShouldKeepAlive in that mode in order to poll the sponsored pid queue in instance info.
* During development there were cases where the shared memory remained despite the instance dissapearing. Added check if the reported pid is actually still alive. There are false positives here if pid was reused.
* On MacOS it's only possible to resize shared memory once. Added code to verify it's only called once.

diff --git a/Engine/Source/Programs/UnrealTraceServer/src/Lifetime.cpp b/Engine/Source/Programs/UnrealTraceServer/src/Lifetime.cpp
--- a/Engine/Source/Programs/UnrealTraceServer/src/Lifetime.cpp
+++ b/Engine/Source/Programs/UnrealTraceServer/src/Lifetime.cpp
@@ -22,10 +22,6 @@ bool FLifetime::ShouldKeepAlive()
 	// Check if all sponsors are alive. If no sponsors are alive try to shutdown
 	// the store if it has no active connections. 
 	const bool bShouldKeepAlive = IsAnySponsorActive() || !ShutdownStoreIfNoConnections();
-	if (!bShouldKeepAlive)
-	{
-		TS_LOG("Terminating server, no sponsors or connections active.");
-	}
 	return bShouldKeepAlive;
 }
 
diff --git a/Engine/Source/Programs/UnrealTraceServer/src/Main.cpp b/Engine/Source/Programs/UnrealTraceServer/src/Main.cpp
--- a/Engine/Source/Programs/UnrealTraceServer/src/Main.cpp
+++ b/Engine/Source/Programs/UnrealTraceServer/src/Main.cpp
@@ -16,11 +16,13 @@
 #endif
 
 #if TS_USING(TS_PLATFORM_LINUX) || TS_USING(TS_PLATFORM_MAC)
+#	include <cstdarg>
+#	include <ctime>
+#	include <pthread.h>
 #	include <pwd.h>
-#	include <semaphore.h>
 #	include <sched.h>
+#	include <semaphore.h>
 #	include <signal.h>
-#	include <cstdarg>
 #	include <sys/file.h>
 #	include <sys/mman.h>
 #	include <sys/stat.h>
@@ -891,6 +893,7 @@ static int MainDaemon(int ArgC, char** ArgV, const FOptions& Options)
 		LifetimeManager.CheckNewSponsors(InstanceInfo);
 		if (!LifetimeManager.ShouldKeepAlive() && Settings->Sponsored)
 		{
+			TS_LOG("Terminating server, no sponsors or connections active.");
 			break;
 		}
 	}
@@ -898,6 +901,8 @@ static int MainDaemon(int ArgC, char** ArgV, const FOptions& Options)
 	// Clean up. We are done here.
 	RemoveFromSystemTray();
 	delete StoreService;
+
+	TS_LOG("Daemon is exiting without errors.");
 	return Result_Ok;
 }
 
@@ -996,6 +1001,11 @@ static int LegacyLockFile()
 ////////////////////////////////////////////////////////////////////////////////
 static int MainKillImpl(int ArgC, char** ArgV, pid_t DaemonPid)
 {
+	if (DaemonPid == 0)
+	{
+		return Result_NoQuitEvent;
+	}
+
 	// Issue the terminate signal
 	TS_LOG("Sending SIGTERM to %d", DaemonPid);
 	if (kill(DaemonPid, SIGTERM) < 0)
@@ -1050,14 +1060,10 @@ static int MainKill(int ArgC, char** ArgV, const FOptions& Options)
 		FMmapScope Buffer(BufferPtr);
 		FInstanceInfo* InstanceInfo = Buffer.As<FInstanceInfo>();
 		
-		// If we've got this far then there's an instance running that is old
-		TS_LOG("Killing an older instance that is already running");
 		int KillRet = MainKillImpl(0, nullptr, InstanceInfo->Pid);
 		if (KillRet == Result_NoQuitEvent)
 		{
-			// If no quit event was found then we shall assume that another new
-			// store instance beat us to it.
-			TS_LOG("Looks like someone else has already taken care of the upgrade");
+			TS_LOG("Looks like someone else has already taken care of stopping");
 			return Result_Ok;
 		}
 
@@ -1097,7 +1103,8 @@ static int MainFork(int ArgC, char** ArgV, const FOptions& Options)
 		FInstanceInfo* InstanceInfo = Buffer.As<FInstanceInfo>();
 
 		// Old enough for this fine establishment?
-		if (!InstanceInfo->IsOlder())
+		bool bExists = kill(InstanceInfo->Pid, 0) == 0;
+		if (bExists && !InstanceInfo->IsOlder())
 		{
 			TS_LOG("Existing instance is the same age or newer");
 
@@ -1133,20 +1140,23 @@ static int MainFork(int ArgC, char** ArgV, const FOptions& Options)
 		}
 
 		// If we've got this far then there's an instance running that is old
-		TS_LOG("Killing an older instance that is already running");
-		int KillRet = MainKillImpl(0, nullptr, InstanceInfo->Pid);
-		if (KillRet == Result_NoQuitEvent)
+		if (bExists)
 		{
-			// If no quit event was found then we shall assume that another new
-			// store instance beat us to it.
-			TS_LOG("Looks like someone else has already taken care of the upgrade");
-			return Result_Ok;
-		}
+			TS_LOG("Killing an older instance (pid %u) that is already running", InstanceInfo->Pid);
+			int KillRet = MainKillImpl(0, nullptr, InstanceInfo->Pid);
+			if (KillRet == Result_NoQuitEvent)
+			{
+				// If no quit event was found then we shall assume that another new
+				// store instance beat us to it.
+				TS_LOG("Looks like someone else has already taken care of the upgrade");
+				return Result_Ok;
+			}
 
-		if (KillRet != Result_Ok)
-		{
-			TS_LOG("Kill attempt failed (ret=%d)", KillRet);
-			return KillRet;
+			if (KillRet != Result_Ok)
+			{
+				TS_LOG("Kill attempt failed (ret=%d)", KillRet);
+				return KillRet;
+			}
 		}
 	}
 #if TS_USING(TS_LEGACY_LOCK_FILE)
@@ -1228,6 +1238,30 @@ static int MainFork(int ArgC, char** ArgV, const FOptions& Options)
 	return Ret;
 }
 
+////////////////////////////////////////////////////////////////////////////////
+static struct FCheckSponsorData {
+	FLifetime* LifetimeManager;
+	FStoreSettings* Settings;
+	FInstanceInfo* InstanceInfo;
+} GCheckSponsorData;
+
+void* CheckSponsors(void* Payload) 
+{
+	FCheckSponsorData* Data = (FCheckSponsorData*) Payload;
+	while(true)
+	{
+		Data->LifetimeManager->CheckNewSponsors(Data->InstanceInfo);
+		if (!Data->LifetimeManager->ShouldKeepAlive() && Data->Settings->Sponsored)
+		{
+			TS_LOG("Terminating server, no sponsors or connections active.");
+			break;
+		}
+		timespec Frequency = { GSponsorCheckFreqSecs, 0 };
+		nanosleep(&Frequency, nullptr);
+	}
+	raise(SIGUSR2);
+}
+
 ////////////////////////////////////////////////////////////////////////////////
 static int MainDaemonImpl(int ArgC, char** ArgV, pid_t ParentPid, const FOptions& Options)
 {
@@ -1235,15 +1269,40 @@ static int MainDaemonImpl(int ArgC, char** ArgV, pid_t ParentPid, const FOptions
 	int Fd = shm_open("/UnrealTraceServer", O_RDWR | O_CREAT | O_CLOEXEC, 0666);
 	if (Fd < 0)
 	{
-		TS_LOG("Unable to create shared memory");
+		TS_LOG("Unable to create shared memory: %u", errno);
 		return Result_SharedMemFail;
 	}
+	OnScopeExit([]() { shm_unlink("/UnrealTraceServer"); });
 
-	// Set the size of the file
+	// Set the size of the file. Note that MacOS allows only setting the size only once.
+#if TS_USING(TS_PLATFORM_MAC)
+	struct stat Stat;
+	if (fstat(Fd,&Stat) == 0)
+	{
+		if (Stat.st_size == 0)
+		{
+			if(ftruncate(Fd, sizeof(FInstanceInfo)) != 0)
+			{
+				TS_LOG("Unable to size shared memory: %d", errno);
+				return Result_SharedMemFail;
+			}
+		}
+		else 
+		{
+			TS_LOG("Shared memory sized: %u", Stat.st_size);
+		}
+	}
+	else 
+	{
+		TS_LOG("Cannot read shared memory stats: %d", errno);
+	}
+#else
 	if (ftruncate(Fd, sizeof(FInstanceInfo)))
 	{
+		TS_LOG("Unable to size shared memory: %d", errno);
 		return Result_SharedMemFail;
 	}
+#endif
 
 	void* BufferPtr = mmap(nullptr, sizeof(FInstanceInfo), PROT_READ | PROT_WRITE, MAP_SHARED, Fd, 0);
 	if (!BufferPtr)
@@ -1311,38 +1370,39 @@ static int MainDaemonImpl(int ArgC, char** ArgV, pid_t ParentPid, const FOptions
 	sigaction(SIGTERM, &SigAction, nullptr);
 	sigaction(SIGINT, &SigAction, nullptr);
 
-	TS_LOG("Entering signal wait loop...");
+	// Create and set a signal handler
 	sigset_t SignalSet;
 	sigemptyset(&SignalSet);
 	sigaddset(&SignalSet, SIGTERM);
 	sigaddset(&SignalSet, SIGKILL);
 	sigaddset(&SignalSet, SIGINT);
+	sigaddset(&SignalSet, SIGUSR2);
 
-	timespec Timeout;
-	Timeout.tv_sec = GSponsorCheckFreqSecs;
+	// Create the lifetime maneger that tracks the sponsor
+	// processes.
 	FLifetime LifetimeManager(StoreService);
 	LifetimeManager.AddPid(SponsorPid);
 
+	// Thread to check lifetime of sponsors
+	pthread_t LifetimeThread;
+	GCheckSponsorData.LifetimeManager = &LifetimeManager;
+	GCheckSponsorData.Settings = Settings;
+	GCheckSponsorData.InstanceInfo = InstanceInfo;
+	OnScopeExit([&LifetimeThread]() { pthread_cancel(LifetimeThread); });
+
+	pthread_create(&LifetimeThread, NULL, CheckSponsors, (void*)&GCheckSponsorData);
+
 	while (true)
 	{
-		siginfo_t SignalInfo;
-		int Ret = sigtimedwait(&SignalSet, &SignalInfo, &Timeout);
-		if (Ret == EINTR)
-		{
-			TS_LOG("Received signal %d", SignalInfo.si_signo);
-			break;
-		}
-
-		LifetimeManager.CheckNewSponsors(InstanceInfo);
-		if (!LifetimeManager.ShouldKeepAlive() && Settings->Sponsored)
+		int Signal = -1;
+		int Ret = sigwait(&SignalSet, &Signal);
+		if (Ret == 0)
 		{
+			TS_LOG("Received signal %d", Signal);
 			break;
 		}
 	}
-
-	// Clean up. We are done here.
-	shm_unlink("/UnrealTraceServer");
-
+	TS_LOG("Daemon is exiting without errors.");
 	return Result_Ok;
 }
 
