Change method of checking for sponsor process lifetime

* Move sponsor process check inside store and use asio tasks to check, instead of launching separate thread on Linux/Mac. There was some strange interaction between standalone pthreads and asio.
* How checking is done is now platform independent, but some platform specific code is needed to trigger exit events.

diff --git a/Engine/Source/Programs/UnrealTraceServer/src/AsioTcpServer.cpp b/Engine/Source/Programs/UnrealTraceServer/src/AsioTcpServer.cpp
--- a/Engine/Source/Programs/UnrealTraceServer/src/AsioTcpServer.cpp
+++ b/Engine/Source/Programs/UnrealTraceServer/src/AsioTcpServer.cpp
@@ -3,6 +3,7 @@
 #include "Pch.h"
 #include "AsioTcpServer.h"
 #include "Foundation.h"
+#include "Logging.h"
 
 ////////////////////////////////////////////////////////////////////////////////
 FAsioTcpServer::FAsioTcpServer(asio::io_context& IoContext)
@@ -74,6 +75,12 @@ bool FAsioTcpServer::StartServer(uint32 Port, uint32 Backlog)
 	TempAcceptor.open(tcp::v4());
 #endif
 
+	// Set the reuse address option (SO_REUSEADDR on *nix) otherwise
+	// when stopping the server the connection is in timed wait prevent
+	// rebinding to that address.
+	asio::socket_base::reuse_address reuse_option(true);
+	TempAcceptor.set_option(reuse_option);
+
 	tcp::endpoint Endpoint(tcp::v4(), uint16(Port));
 	TempAcceptor.bind(Endpoint);
 	TempAcceptor.listen(Backlog);
@@ -121,6 +128,15 @@ void FAsioTcpServer::AsyncAccept()
 			AsyncAccept();
 			return;
 		}
+		
+		if (ErrorCode == asio::error::operation_aborted)
+		{
+			TS_LOG("Listening cancelled, closing port...");
+		}
+		else if (ErrorCode)
+		{
+			TS_LOG("Accept failed with error %u, closing port...", ErrorCode);
+		}
 
 		Close();
 	});
diff --git a/Engine/Source/Programs/UnrealTraceServer/src/InstanceInfo.cpp b/Engine/Source/Programs/UnrealTraceServer/src/InstanceInfo.cpp
--- a/Engine/Source/Programs/UnrealTraceServer/src/InstanceInfo.cpp
+++ b/Engine/Source/Programs/UnrealTraceServer/src/InstanceInfo.cpp
@@ -40,4 +40,26 @@ bool FInstanceInfo::IsOlder() const
 	bool bIsOlder = false;
 	bIsOlder |= (Version < FInstanceInfo::CurrentVersion);
 	return bIsOlder;
-}
\ No newline at end of file
+}
+
+////////////////////////////////////////////////////////////////////////////////
+bool FInstanceInfo::AddSponsor(uint32 Pid)
+{
+	if (!Pid)
+	{
+		return true;
+	}
+
+	for (auto& PidEntry : SponsorPids)
+	{
+		if (uint32_t ThisPid = PidEntry.load(std::memory_order_relaxed); !ThisPid)
+		{
+			if (PidEntry.compare_exchange_strong(ThisPid, Pid, std::memory_order_relaxed))
+			{
+				return true;
+			}
+		}
+	}
+	return false;
+}
+
diff --git a/Engine/Source/Programs/UnrealTraceServer/src/InstanceInfo.h b/Engine/Source/Programs/UnrealTraceServer/src/InstanceInfo.h
--- a/Engine/Source/Programs/UnrealTraceServer/src/InstanceInfo.h
+++ b/Engine/Source/Programs/UnrealTraceServer/src/InstanceInfo.h
@@ -18,6 +18,8 @@ public:
 	void				Set();
 	void				WaitForReady() const;
 	bool				IsOlder() const;
+	bool				AddSponsor(uint32 Pid);
+
 	std::atomic<uint32> Published;
 	uint32				Version;
 	uint32				Pid;
diff --git a/Engine/Source/Programs/UnrealTraceServer/src/Lifetime.cpp b/Engine/Source/Programs/UnrealTraceServer/src/Lifetime.cpp
--- a/Engine/Source/Programs/UnrealTraceServer/src/Lifetime.cpp
+++ b/Engine/Source/Programs/UnrealTraceServer/src/Lifetime.cpp
@@ -4,25 +4,61 @@
 #include "InstanceInfo.h"
 #include "Logging.h"
 #include "StoreService.h"
+#include "StoreSettings.h"
 
 #if TS_USING(TS_PLATFORM_LINUX) || TS_USING(TS_PLATFORM_MAC)
 #include <signal.h>
 #endif
 
+#if TS_USING(TS_PLATFORM_WINDOWS)
+static HANDLE GQuitEvent = NULL;
+extern const wchar_t* GQuitEventName;
+#endif
+
+constexpr uint32 GSponsorCheckFreqSecs = 5;
+
 ////////////////////////////////////////////////////////////////////////////////
-FLifetime::FLifetime(class FStoreService* InStoreService)
-	: StoreService(InStoreService)
+FLifetime::FLifetime(asio::io_context& IoContext, 
+					 FStoreService* InStoreService, 
+					 FStoreSettings* InSettings, 
+					 FInstanceInfo* InInstanceInfo)
+	: FAsioTickable(IoContext)
+	, StoreService(InStoreService)
+	, Settings(InSettings)
+	, InstanceInfo(InInstanceInfo)
 {
-	check(StoreService && InstanceInfo);
+	check(StoreService);
+	check(InstanceInfo);
+	check(Settings);
+
+#if TS_USING(TS_PLATFORM_WINDOWS)
+	GQuitEvent = OpenEventW(EVENT_MODIFY_STATE, FALSE, GQuitEventName);
+	check(GQuitEvent != NULL);
+#endif
+
+	StartTick(GSponsorCheckFreqSecs * 1000);
 }
 
 ////////////////////////////////////////////////////////////////////////////////
-bool FLifetime::ShouldKeepAlive()
+void FLifetime::OnTick()
 {
-	// Check if all sponsors are alive. If no sponsors are alive try to shutdown
-	// the store if it has no active connections. 
-	const bool bShouldKeepAlive = IsAnySponsorActive() || !ShutdownStoreIfNoConnections();
-	return bShouldKeepAlive;
+	CheckNewSponsors(InstanceInfo);
+	// We need to call IsAnySponsorsActive regularly even if sponsored mode is not 
+	// activated. This is because otherwise the SponsorHandles array would build 
+	// up indefinetly and the risk of reuse of pid increase.
+	if (!IsAnySponsorActive() && Settings->Sponsored)
+	{
+		if (ShutdownStoreIfNoConnections())
+		{
+			TS_LOG("Terminating server, no sponsors or connections active.");
+#if TS_USING(TS_PLATFORM_WINDOWS)
+			SetEvent(GQuitEvent);
+#else
+			kill(getpid(), SIGTERM);
+#endif
+			StopTick();
+		}
+	}
 }
 
 ////////////////////////////////////////////////////////////////////////////////
diff --git a/Engine/Source/Programs/UnrealTraceServer/src/Lifetime.h b/Engine/Source/Programs/UnrealTraceServer/src/Lifetime.h
--- a/Engine/Source/Programs/UnrealTraceServer/src/Lifetime.h
+++ b/Engine/Source/Programs/UnrealTraceServer/src/Lifetime.h
@@ -3,24 +3,32 @@
 #pragma once
 
 #include "Foundation.h"
+#include "AsioTickable.h"
 
 ////////////////////////////////////////////////////////////////////////////////
 
-class FLifetime
+class FLifetime : public FAsioTickable
 {
 public:
-						FLifetime(class FStoreService* Service);
-	bool				ShouldKeepAlive();
-	void				CheckNewSponsors(struct FInstanceInfo* InstanceInfo);
-	void				AddPid(uint32);
+						FLifetime(
+							asio::io_context& IoContext, 
+							class FStoreService* Service, 
+							class FStoreSettings* Settings, 
+							struct FInstanceInfo* InstanceInfo
+						);
 
 private:
+	void				CheckNewSponsors(struct FInstanceInfo* InstanceInfo);
+	void				AddPid(uint32);
 	bool				ShutdownStoreIfNoConnections();
 	bool				IsAnySponsorActive();
+	virtual void		OnTick() override;
 
 	using FProcHandle = void*;
 
 	TArray<FProcHandle> SponsorHandles;
 	class FStoreService* StoreService;
+	class FStoreSettings* Settings;
+	struct FInstanceInfo* InstanceInfo;
 };
 
diff --git a/Engine/Source/Programs/UnrealTraceServer/src/Main.cpp b/Engine/Source/Programs/UnrealTraceServer/src/Main.cpp
--- a/Engine/Source/Programs/UnrealTraceServer/src/Main.cpp
+++ b/Engine/Source/Programs/UnrealTraceServer/src/Main.cpp
@@ -41,9 +41,6 @@
 // Check for legacy lock files on Linux and Mac 
 #define TS_LEGACY_LOCK_FILE TS_ON
 
-// Frequency between sponsor lifetime checks
-constexpr uint32 GSponsorCheckFreqSecs = 5;
-
 // {{{1 misc -------------------------------------------------------------------
 
 ////////////////////////////////////////////////////////////////////////////////
@@ -502,7 +499,7 @@ static bool LaunchUnelevated(const wchar_t* Binary, wchar_t* CommandLine)
 ////////////////////////////////////////////////////////////////////////////////
 static const wchar_t*	GIpcName					= L"Local\\UnrealTraceInstance";
 static const int32		GIpcSize					= 4 << 10;
-static const wchar_t*	GQuitEventName				= L"Local\\UnrealTraceEvent";
+const wchar_t*			GQuitEventName				= L"Local\\UnrealTraceEvent";
 static const wchar_t*	GBegunEventName				= L"Local\\UnrealTraceEventBegun";
 static int				MainDaemon(int, char**, const FOptions&);
 void					AddToSystemTray(FStoreService&);
@@ -862,9 +859,11 @@ static int MainDaemon(int ArgC, char** ArgV, const FOptions& Options)
 		TS_LOG("Error: Deamon is configured to run in sponsored mode, but no sponsor pid has been specified.");
 		return Result_InvalidArgError;
 	}
+	InstanceInfo->AddSponsor(ParentPid);
 
 	// Fire up the store
-	FStoreService* StoreService = FStoreService::Create(Settings);
+	FStoreService* StoreService = FStoreService::Create(Settings, InstanceInfo);
+	OnScopeExit([StoreService]() { delete StoreService; });
 
 	// Let every one know we've started.
 	{
@@ -880,27 +879,10 @@ static int MainDaemon(int ArgC, char** ArgV, const FOptions& Options)
 	AddToSystemTray(*StoreService);
 
 	// Wait to be told to resign.
-	FLifetime LifetimeManager(StoreService);
-	LifetimeManager.AddPid(ParentPid);
-
-	while (true)
-	{
-		DWORD Result = WaitForSingleObject(QuitEvent, GSponsorCheckFreqSecs * 1000);
-		if (Result != WAIT_TIMEOUT)
-		{
-			break;
-		}
-		LifetimeManager.CheckNewSponsors(InstanceInfo);
-		if (!LifetimeManager.ShouldKeepAlive() && Settings->Sponsored)
-		{
-			TS_LOG("Terminating server, no sponsors or connections active.");
-			break;
-		}
-	}
+	WaitForSingleObject(QuitEvent, INFINITE);
 
 	// Clean up. We are done here.
 	RemoveFromSystemTray();
-	delete StoreService;
 
 	TS_LOG("Daemon is exiting without errors.");
 	return Result_Ok;
@@ -1238,30 +1220,6 @@ static int MainFork(int ArgC, char** ArgV, const FOptions& Options)
 	return Ret;
 }
 
-////////////////////////////////////////////////////////////////////////////////
-static struct FCheckSponsorData {
-	FLifetime* LifetimeManager;
-	FStoreSettings* Settings;
-	FInstanceInfo* InstanceInfo;
-} GCheckSponsorData;
-
-void* CheckSponsors(void* Payload) 
-{
-	FCheckSponsorData* Data = (FCheckSponsorData*) Payload;
-	while(true)
-	{
-		Data->LifetimeManager->CheckNewSponsors(Data->InstanceInfo);
-		if (!Data->LifetimeManager->ShouldKeepAlive() && Data->Settings->Sponsored)
-		{
-			TS_LOG("Terminating server, no sponsors or connections active.");
-			break;
-		}
-		timespec Frequency = { GSponsorCheckFreqSecs, 0 };
-		nanosleep(&Frequency, nullptr);
-	}
-	raise(SIGUSR2);
-}
-
 ////////////////////////////////////////////////////////////////////////////////
 static int MainDaemonImpl(int ArgC, char** ArgV, pid_t ParentPid, const FOptions& Options)
 {
@@ -1337,8 +1295,11 @@ static int MainDaemonImpl(int ArgC, char** ArgV, pid_t ParentPid, const FOptions
 		TS_LOG("Error: Deamon is configured to run in sponsored mode, but no sponsor pid has been specified.");
 		return Result_InvalidArgError;
 	}
+	// Add given sponsor pid regardless, in case we suddenly enable
+	// sponsor mode 
+	InstanceInfo->AddSponsor(SponsorPid);
 
-	StoreService = FStoreService::Create(Settings);
+	StoreService = FStoreService::Create(Settings, InstanceInfo);
 	OnScopeExit([StoreService] () { delete StoreService; });
 
 	// Fill out the shared memory with details about this instance
@@ -1376,21 +1337,6 @@ static int MainDaemonImpl(int ArgC, char** ArgV, pid_t ParentPid, const FOptions
 	sigaddset(&SignalSet, SIGTERM);
 	sigaddset(&SignalSet, SIGKILL);
 	sigaddset(&SignalSet, SIGINT);
-	sigaddset(&SignalSet, SIGUSR2);
-
-	// Create the lifetime maneger that tracks the sponsor
-	// processes.
-	FLifetime LifetimeManager(StoreService);
-	LifetimeManager.AddPid(SponsorPid);
-
-	// Thread to check lifetime of sponsors
-	pthread_t LifetimeThread;
-	GCheckSponsorData.LifetimeManager = &LifetimeManager;
-	GCheckSponsorData.Settings = Settings;
-	GCheckSponsorData.InstanceInfo = InstanceInfo;
-	OnScopeExit([&LifetimeThread]() { pthread_cancel(LifetimeThread); });
-
-	pthread_create(&LifetimeThread, NULL, CheckSponsors, (void*)&GCheckSponsorData);
 
 	while (true)
 	{
diff --git a/Engine/Source/Programs/UnrealTraceServer/src/StoreService.cpp b/Engine/Source/Programs/UnrealTraceServer/src/StoreService.cpp
--- a/Engine/Source/Programs/UnrealTraceServer/src/StoreService.cpp
+++ b/Engine/Source/Programs/UnrealTraceServer/src/StoreService.cpp
@@ -3,6 +3,8 @@
 #include "Pch.h"
 #include "Asio.h"
 #include "AsioContext.h"
+#include "InstanceInfo.h"
+#include "Lifetime.h"
 #include "Recorder.h"
 #include "Store.h"
 #include "StoreCborServer.h"
@@ -13,20 +15,22 @@
 struct FStoreServiceImpl
 {
 public:
-							FStoreServiceImpl(FStoreSettings* Settings);
+							FStoreServiceImpl(FStoreSettings* Settings, FInstanceInfo* InstanceInfo);
 							~FStoreServiceImpl();
 	FAsioContext			Context;
 	FStore					Store;
 	FRecorder				Recorder;
 	FStoreCborServer		CborServer;
+	FLifetime				LifetimeManager;
 };
 
 ////////////////////////////////////////////////////////////////////////////////
-FStoreServiceImpl::FStoreServiceImpl(FStoreSettings* Settings)
+FStoreServiceImpl::FStoreServiceImpl(FStoreSettings* Settings, FInstanceInfo* InstanceInfo)
 : Context(Settings->ThreadCount)
 , Store(Context.Get(), Settings)
 , Recorder(Context.Get(), Store)
 , CborServer(Context.Get(), Settings, Store, Recorder)
+, LifetimeManager(Context.Get(), (FStoreService*)this, Settings, InstanceInfo)
 {
 	if (Settings->RecorderPort >= 0)
 	{
@@ -39,6 +43,7 @@ FStoreServiceImpl::FStoreServiceImpl(FStoreSettings* Settings)
 FStoreServiceImpl::~FStoreServiceImpl()
 {
 	asio::post(Context.Get(), [this] () {
+		LifetimeManager.StopTick();
 		CborServer.Close();
 		Recorder.Close();
 		Store.Close();
@@ -49,7 +54,7 @@ FStoreServiceImpl::~FStoreServiceImpl()
 
 
 ////////////////////////////////////////////////////////////////////////////////
-FStoreService* FStoreService::Create(FStoreSettings* Settings)
+FStoreService* FStoreService::Create(FStoreSettings* Settings, FInstanceInfo* InstanceInfo)
 {
 	if (Settings->ThreadCount <= 0)
 	{
@@ -59,7 +64,7 @@ FStoreService* FStoreService::Create(FStoreSettings* Settings)
 	// TODO: not thread safe yet
 	Settings->ThreadCount = 1;
 
-	return (FStoreService*) new FStoreServiceImpl(Settings);
+	return (FStoreService*) new FStoreServiceImpl(Settings, InstanceInfo);
 }
 
 ////////////////////////////////////////////////////////////////////////////////
diff --git a/Engine/Source/Programs/UnrealTraceServer/src/StoreService.h b/Engine/Source/Programs/UnrealTraceServer/src/StoreService.h
--- a/Engine/Source/Programs/UnrealTraceServer/src/StoreService.h
+++ b/Engine/Source/Programs/UnrealTraceServer/src/StoreService.h
@@ -11,7 +11,7 @@ class FStoreService
 {
 public:
 							~FStoreService() = default;
-	static FStoreService*	Create(FStoreSettings* Desc);
+	static FStoreService*	Create(FStoreSettings* Desc, struct FInstanceInfo* InstanceInfo);
 	void					operator delete (void* Addr);
 	uint32					GetPort() const;
 	uint32					GetRecorderPort() const;
