Fix crash when looking for user home directory

When UTS was looking for the users home directory getpwuid was used. In certain scenarios this doesn't work and a null pointer is returned. Man pages suggest reading the HOME environment variable is more correct, so with this change we use that primarily and fall back to getpwuid method.

diff --git a/Engine/Source/Programs/UnrealTraceServer/src/Main.cpp b/Engine/Source/Programs/UnrealTraceServer/src/Main.cpp
--- a/Engine/Source/Programs/UnrealTraceServer/src/Main.cpp
+++ b/Engine/Source/Programs/UnrealTraceServer/src/Main.cpp
@@ -122,9 +122,23 @@ static void GetUnrealTraceHome(FPath& Out, bool Make=false)
 	Out = Buffer;
 	Out /= "UnrealEngine/Common/UnrealTrace";
 #else
-	int UserId = getuid();
-	const passwd* Passwd = getpwuid(UserId);
-	Out = Passwd->pw_dir;
+	if (const char* Home = getenv("HOME"))
+	{
+		Out = Home;
+	}
+	else
+	{
+		int UserId = getuid();
+		const passwd* Passwd = getpwuid(UserId);
+		if (Passwd)
+		{
+			Out = Passwd->pw_dir;
+		}
+		else 
+		{
+			Out = "/";
+		}
+	}
 	Out /= "UnrealEngine/UnrealTrace";
 #endif
 
