Fix some compilation warnings

diff --git a/Engine/Source/Programs/UnrealTraceServer/src/CborPayload.h b/Engine/Source/Programs/UnrealTraceServer/src/CborPayload.h
--- a/Engine/Source/Programs/UnrealTraceServer/src/CborPayload.h
+++ b/Engine/Source/Programs/UnrealTraceServer/src/CborPayload.h
@@ -195,7 +195,7 @@ inline int64 FResponse::GetInteger(const char* Key, int64 Default) const
 	return GetValue(
 		Key,
 		Default,
-		[this] (const FCborContext& Context)
+		[] (const FCborContext& Context)
 		{
 			return Context.AsInteger();
 		}
@@ -210,7 +210,7 @@ inline FStringView FResponse::GetString(const char* Key, const char (&Default)[N
 	return GetValue(
 		Key,
 		DefaultView,
-		[this, DefaultView] (const FCborContext& Context)
+		[DefaultView] (const FCborContext& Context)
 		{
 			if (Context.GetType() == ECborType::String)
 			{
diff --git a/Engine/Source/Programs/UnrealTraceServer/src/Main.cpp b/Engine/Source/Programs/UnrealTraceServer/src/Main.cpp
--- a/Engine/Source/Programs/UnrealTraceServer/src/Main.cpp
+++ b/Engine/Source/Programs/UnrealTraceServer/src/Main.cpp
@@ -194,7 +194,7 @@ public:
 
 	ParseResults Parse(int ArgC, char** ArgV)
 	{
-		Parsed = MoveTemp(Options.parse(ArgC, ArgV));
+		Parsed = Options.parse(ArgC, ArgV);
 		const bool bCommandOk = Parsed["command"].count() == 1;
 		const bool bHelp = Parsed["help"].count() > 0;
 
@@ -310,14 +310,14 @@ private:
 };
 
 FOptions::FCommandHelp FOptions::CommandHelp[] = {
-	"fork", "	Starts a background server, upgrading any existing instance. ", 
+	{"fork", "	Starts a background server, upgrading any existing instance. ", 
 				"Checks if there is an existing instance running. If the running version is the same "
 				"version or newer that instance is used. If a sponsor pid is specified that pid is added "
-				"to the running instance.", {"settings", "sponsor"},
-	"daemon", "The mode that a background server runs in. ", "", {"sponsor"},
-	"kill", "	Shuts down a currently running instance. ", "", {},
-	"test", "	Run tests. ", "", {},
-	nullptr, nullptr, nullptr, {}
+				"to the running instance.", {"settings", "sponsor"}},
+	{"daemon", "The mode that a background server runs in. ", "", {"sponsor"}},
+	{"kill", "	Shuts down a currently running instance. ", "", {}},
+	{"test", "	Run tests. ", "", {}},
+	{nullptr, nullptr, nullptr, {}}
 };
 
 // {{{1 return codes -----------------------------------------------------------
@@ -1399,10 +1399,10 @@ int main(int ArgC, char** ArgV)
 		const char*	Verb;
 		int			(*Entry)(int, char**, const FOptions&);
 	} Dispatches[] = {
-		"fork",		MainFork,
-		"daemon",	MainDaemon,
-		"test",		MainTest,
-		"kill",		MainKill,
+		{"fork",		MainFork},
+		{"daemon",		MainDaemon},
+		{"test",		MainTest},
+		{"kill",		MainKill},
 	};
 
 	for (const auto& Dispatch : Dispatches)
diff --git a/Engine/Source/Programs/UnrealTraceServer/src/Store.cpp b/Engine/Source/Programs/UnrealTraceServer/src/Store.cpp
--- a/Engine/Source/Programs/UnrealTraceServer/src/Store.cpp
+++ b/Engine/Source/Programs/UnrealTraceServer/src/Store.cpp
@@ -3,6 +3,7 @@
 #include "Pch.h"
 #include "Store.h"
 #include "StoreSettings.h"
+#include <functional>
 
 #if TS_USING(TS_PLATFORM_LINUX)
 #   include <sys/inotify.h>
@@ -143,12 +144,12 @@ public:
 	void ProcessChanges(size_t EventCount, void* EventPaths, const FSEventStreamEventFlags EventFlags[])
 	{
 		bool bWatchedEvent = false;
-		const char** EventPathsArray = (const char**) EventPaths;
 		for (size_t EventIdx = 0; EventIdx < EventCount; ++EventIdx)
 		{
-			const char* Path = EventPathsArray[EventIdx];
 			const FSEventStreamEventFlags& Flags = EventFlags[EventIdx];
 #if TS_USING(TS_DEBUG_FS_EVENTS)
+			const char** EventPathsArray = (const char**) EventPaths;
+			const char* Path = EventPathsArray[EventIdx];
 			printf("Recieved file event (%d) ", EventIdx);
 			if (Path != nullptr)
 				printf("on '%s': ", Path);
@@ -353,10 +354,12 @@ FStore::FMount* FStore::FMount::Create(FStore* InParent, asio::io_context& InIoC
 
 ////////////////////////////////////////////////////////////////////////////////
 FStore::FMount::FMount(FStore* InParent, asio::io_context& InIoContext, const fs::path& InDir)
-: Id(QuickStoreHash(InDir.c_str()))
-, Dir(InDir)
+: Dir(InDir)
+, Id(QuickStoreHash(InDir.c_str()))
 , Parent(InParent)
+#if TS_USING(TS_PLATFORM_WINDOWS)
 , IoContext(InIoContext)
+#endif
 {
 
 #if TS_USING(TS_PLATFORM_WINDOWS)
@@ -676,7 +679,7 @@ FStore::FNewTrace FStore::CreateTrace()
 	for (uint32 Index = 0; std::filesystem::is_regular_file(TracePath); ++Index)
 	{
 		char FilenameIndexed[64];
-		std::sprintf(FilenameIndexed, "%s_%02d.utrace", Prefix, Index);
+		std::snprintf(FilenameIndexed, 64, "%s_%02d.utrace", Prefix, Index);
 		TracePath.replace_filename(FPath(FilenameIndexed));
 	}
 
diff --git a/Engine/Source/Programs/UnrealTraceServer/src/Store.h b/Engine/Source/Programs/UnrealTraceServer/src/Store.h
--- a/Engine/Source/Programs/UnrealTraceServer/src/Store.h
+++ b/Engine/Source/Programs/UnrealTraceServer/src/Store.h
@@ -67,12 +67,15 @@ public:
 		void			ClearTraces();
 		uint32			Refresh();
 		void			WatchDir();
+		
 		FPath			Dir;
 		TArray<FTrace*>	Traces;
 		uint32			Id;
 		FDirWatcher*	DirWatcher = nullptr;
 		FStore*			Parent = nullptr;
+#if TS_USING(TS_PLATFORM_WINDOWS)
 		asio::io_context& IoContext;
+#endif
 	};
 
 	void SetupMounts();
