Fix checking privileged applications.

If UTS was launched by a process with elevated privileges and provided as a sponsor process, UTS was not able to open a handle to the sponsor which resulted in immediate exit. We only need the exit status of the process so opening the sponsor with PROCESS_QUERY_LIMITED_INFORMATION is enough.

diff --git a/Engine/Source/Programs/UnrealTraceServer/src/Lifetime.cpp b/Engine/Source/Programs/UnrealTraceServer/src/Lifetime.cpp
--- a/Engine/Source/Programs/UnrealTraceServer/src/Lifetime.cpp
+++ b/Engine/Source/Programs/UnrealTraceServer/src/Lifetime.cpp
@@ -85,7 +85,12 @@ void FLifetime::AddPid(uint32 Pid)
 		return;
 	}
 #if TS_USING(TS_PLATFORM_WINDOWS)
-	FProcHandle ProcessHandle = OpenProcess(PROCESS_QUERY_INFORMATION | SYNCHRONIZE, FALSE, Pid);
+	FProcHandle ProcessHandle = OpenProcess(PROCESS_QUERY_LIMITED_INFORMATION | SYNCHRONIZE, FALSE, Pid);
+	if (!ProcessHandle)
+	{
+		TS_LOG("Failed to open sponsor process (error code: %d), will not be able to use as sponsor.", GetLastError());
+		return;
+	}
 #else
 	FProcHandle ProcessHandle = FProcHandle(intptr_t(Pid));
 #endif
