Add command line override for sponsor mode.

* Adds command line argument for setting sponsor mode "--sponsor-mode".
* Added additional information in command line help regarding settings.

diff --git a/Engine/Source/Programs/UnrealTraceServer/src/Foundation.h b/Engine/Source/Programs/UnrealTraceServer/src/Foundation.h
--- a/Engine/Source/Programs/UnrealTraceServer/src/Foundation.h
+++ b/Engine/Source/Programs/UnrealTraceServer/src/Foundation.h
@@ -84,7 +84,7 @@ struct FStringView
 
 ////////////////////////////////////////////////////////////////////////////////
 struct FString
-	: protected std::string
+	: public std::string
 {
 	typedef std::string Super;
 	using Super::Super;
diff --git a/Engine/Source/Programs/UnrealTraceServer/src/Main.cpp b/Engine/Source/Programs/UnrealTraceServer/src/Main.cpp
--- a/Engine/Source/Programs/UnrealTraceServer/src/Main.cpp
+++ b/Engine/Source/Programs/UnrealTraceServer/src/Main.cpp
@@ -162,11 +162,7 @@ public:
 	};
 
 	FOptions()
-		: Options("UnrealTraceServer", "Unreal Trace Server"
-			"\n\nUnrealTraceServer acts as a hub between runtimes that are tracing performance "
-			"instrumentation and tools like Unreal Insights that consume and present that "
-			"data for analysis. TCP ports 1981 and 1989 are used, where the former receives "
-			"trace data, and the latter is used by tools to query the server's store.\n")
+		: Options("UnrealTraceServer", FOptions::GetMainHelpString())
 		, CurrentCommand(nullptr)
 	{
 		// Positional arguments
@@ -177,7 +173,7 @@ public:
 		Options.positional_help("<cmd>");
 
 		Options.add_options()
-			("help", "Prints help message for each command");
+			("h,help", "Prints help message for each command");
 
 		// Fork and daemon options
 		Options.add_options("settings")
@@ -189,7 +185,7 @@ public:
 		// AddProc options
 		Options.add_options("sponsor")
 			("sponsor", "Pid to add as sponsor. Required if running in sponsored mode.", cxxopts::value<uint32>()->default_value("0"))
-			;
+			("sponsor-mode", "Set sponsor mode. Zero to disable, non-zero to enable.", cxxopts::value<uint32>());
 	}
 
 	ParseResults Parse(int ArgC, char** ArgV)
@@ -261,6 +257,12 @@ public:
 		{
 			Settings->StoreDir = Value;
 		}
+
+		if (Parsed.count("sponsor-mode"))
+		{
+			uint32 Value = Parsed["sponsor-mode"].as<uint32>();
+			Settings->Sponsored = Value;
+		}
 	}
 
 	bool GetSponsorPid(uint32& OutSponsorPid) const
@@ -279,6 +281,7 @@ public:
 		if (!CurrentCommand)
 		{
 			HelpText = Options.help({"dummy"});
+			HelpText += "Use \"<Command> --help\" to get extended help about specific commands.\n\n";
 		}
 		else
 		{
@@ -295,8 +298,29 @@ public:
 
 
 private:
+
 	cxxopts::Options Options;
 	cxxopts::ParseResult Parsed;
+
+	static FString GetMainHelpString()
+	{
+		FString Help = "\nUnrealTraceServer acts as a hub between runtimes that are tracing performance "
+			"instrumentation and tools like Unreal Insights that consume and present that "
+			"data for analysis. TCP ports 1981 and 1989 are used, where the former receives "
+			"trace data, and the latter is used by tools to query the server's store.\n\n"
+			"UnrealTraceServer will look for a configuration file in \"";
+
+		// Add the settings file path
+#if TS_USING(TS_PLATFORM_WINDOWS)
+		Help += "%LOCALAPPDATA%/UnrealEngine/Common/UnrealTrace/Settings.ini";
+#else
+		Help += "~/UnrealEngine/UnrealTrace/Settings.ini";
+#endif
+		Help += "\". If no file is found default values are used. "
+			"Command line arguments takes precedence over values in the configuration file.\n";
+			
+		return Help;
+	}
 	
 	static struct FCommandHelp
 	{
@@ -306,6 +330,7 @@ private:
 		std::vector<const char*> HelpGroups;
 	} CommandHelp[];
 
+	FString ConfigurationHelp;
 	const FCommandHelp* CurrentCommand;
 };
 
