Add option to detach from console

Add option in UnrealTraceServer to detach from current console. Automatically apply this option when launching a forked server to avoid stopping the process when terminal window closes.

diff --git a/Engine/Source/Programs/UnrealTraceServer/src/Main.cpp b/Engine/Source/Programs/UnrealTraceServer/src/Main.cpp
--- a/Engine/Source/Programs/UnrealTraceServer/src/Main.cpp
+++ b/Engine/Source/Programs/UnrealTraceServer/src/Main.cpp
@@ -173,7 +173,8 @@ public:
 		Options.positional_help("<cmd>");
 
 		Options.add_options()
-			("h,help", "Prints help message for each command");
+			("h,help", "Prints help message for each command")
+			("d,detach", "Detach from console when started.");
 
 		// Fork and daemon options
 		Options.add_options("settings")
@@ -275,6 +276,11 @@ public:
 		return false;
 	}
 
+	bool GetDetach() const
+	{
+		return Parsed["detach"].count() > 0;
+	}
+
 	void PrintHelp() const
 	{
 		std::string HelpText;
@@ -759,7 +765,7 @@ static int MainFork(int ArgC, char** ArgV, const FOptions& Options)
 #if TS_USING(TS_DAEMON_THREAD)
 	std::thread DaemonThread([=] () { MainDaemon(ArgC, ArgV, Options); });
 #else
-	std::wstring CommandLine = L"UnrealTraceServer.exe daemon";
+	std::wstring CommandLine = L"UnrealTraceServer.exe daemon -d ";
 
 	auto ContainsSpace = [](LPCWSTR Arg) -> bool {
 		for (uint32 c = 0, Len = (uint32)wcslen(Arg); c < Len; ++c)
@@ -904,6 +910,12 @@ static int MainDaemon(int ArgC, char** ArgV, const FOptions& Options)
 	FStoreService* StoreService = FStoreService::Create(Settings, InstanceInfo);
 	OnScopeExit([StoreService]() { delete StoreService; });
 
+	if (Options.GetDetach())
+	{
+		TS_LOG("Detaching from console now.");
+		FreeConsole();
+	}
+
 	// Let every one know we've started.
 	{
 		FWinHandle BegunEvent = CreateEventW(nullptr, TRUE, FALSE, GBegunEventName);
