Fix compile error on Linux for UTS.

The io context member is also used on Linux.

diff --git a/Engine/Source/Programs/UnrealTraceServer/src/Store.cpp b/Engine/Source/Programs/UnrealTraceServer/src/Store.cpp
--- a/Engine/Source/Programs/UnrealTraceServer/src/Store.cpp
+++ b/Engine/Source/Programs/UnrealTraceServer/src/Store.cpp
@@ -357,7 +357,7 @@ FStore::FMount::FMount(FStore* InParent, asio::io_context& InIoContext, const fs
 : Dir(InDir)
 , Id(QuickStoreHash(InDir.c_str()))
 , Parent(InParent)
-#if TS_USING(TS_PLATFORM_WINDOWS)
+#if TS_USING(TS_PLATFORM_WINDOWS) || TS_USING(TS_PLATFORM_LINUX)
 , IoContext(InIoContext)
 #endif
 {
diff --git a/Engine/Source/Programs/UnrealTraceServer/src/Store.h b/Engine/Source/Programs/UnrealTraceServer/src/Store.h
--- a/Engine/Source/Programs/UnrealTraceServer/src/Store.h
+++ b/Engine/Source/Programs/UnrealTraceServer/src/Store.h
@@ -73,7 +73,7 @@ public:
 		uint32			Id;
 		FDirWatcher*	DirWatcher = nullptr;
 		FStore*			Parent = nullptr;
-#if TS_USING(TS_PLATFORM_WINDOWS)
+#if TS_USING(TS_PLATFORM_WINDOWS) || TS_USING(TS_PLATFORM_LINUX)
 		asio::io_context& IoContext;
 #endif
 	};
