Moved code handling socket keep-alive settings

When setting keep-alive flags for FRecorder sockets, the operation were made on socket object that had already been moved making them ineffective. This code is now moved to FAsioSocket class and exposed as a platform independent optional argument.

diff --git a/Engine/Source/Programs/UnrealTraceServer/src/AsioIoable.cpp b/Engine/Source/Programs/UnrealTraceServer/src/AsioIoable.cpp
--- a/Engine/Source/Programs/UnrealTraceServer/src/AsioIoable.cpp
+++ b/Engine/Source/Programs/UnrealTraceServer/src/AsioIoable.cpp
@@ -2,6 +2,7 @@
 
 #include "Pch.h"
 #include "AsioIoable.h"
+#include "Logging.h"
 
 ////////////////////////////////////////////////////////////////////////////////
 bool FAsioIoable::SetSink(FAsioIoSink* Ptr, uint32 Id)
@@ -24,10 +25,20 @@ void FAsioIoable::OnIoComplete(const asio::error_code& ErrorCode, int32 Size)
 		return;
 	}
 
-#if TS_USING(TS_BUILD_DEBUG) && 0
-	std::string ErrorMessage;
+#if TS_USING(TS_BUILD_DEBUG) && 1
+	if (ErrorCode && ErrorCode != asio::error::eof)
 	{
-		ErrorMessage = ErrorCode.message();
+		switch (ErrorCode.value())
+		{
+		case asio::error::eof: 
+			break;
+		case asio::error::connection_aborted:
+		case asio::error::connection_reset:
+			TS_LOG("Connection closed (object %p)", this);
+			break;
+		default:
+			TS_LOG("IO error (object %p): (Code %u) %s", this, ErrorCode, ErrorCode.message().c_str());
+		}
 	}
 #endif
 
diff --git a/Engine/Source/Programs/UnrealTraceServer/src/AsioSocket.cpp b/Engine/Source/Programs/UnrealTraceServer/src/AsioSocket.cpp
--- a/Engine/Source/Programs/UnrealTraceServer/src/AsioSocket.cpp
+++ b/Engine/Source/Programs/UnrealTraceServer/src/AsioSocket.cpp
@@ -3,11 +3,72 @@
 #include "Pch.h"
 #include "AsioSocket.h"
 #include "Foundation.h"
+#include "Logging.h"
+
+#if TS_USING(TS_PLATFORM_WINDOWS)
+#	include <Mstcpip.h>
+#endif
 
 ////////////////////////////////////////////////////////////////////////////////
-FAsioSocket::FAsioSocket(asio::ip::tcp::socket& InSocket)
+FAsioSocket::FAsioSocket(asio::ip::tcp::socket& InSocket, FKeepAliveConfig& Config)
 : Socket(MoveTemp(InSocket))
 {
+	if (Config.Enabled)
+	{
+		int Error(0), ErrorCode(0);
+
+#if TS_USING(TS_PLATFORM_WINDOWS)
+		tcp_keepalive KeepAlive =
+		{
+			Config.Enabled,		// on
+			Config.TimeoutSec * 1000,	// timeout_ms
+			Config.IntervalSec * 1000,	// interval_ms
+		};
+
+		DWORD BytesReturned;
+		Error = WSAIoctl(
+			Socket.native_handle(),
+			SIO_KEEPALIVE_VALS,
+			&KeepAlive, sizeof(KeepAlive),
+			nullptr, 0,
+			&BytesReturned,
+			nullptr,
+			nullptr
+		);
+		if (Error)
+		{
+			ErrorCode = WSAGetLastError();
+		}
+
+#elif TS_USING(TS_PLATFORM_MAC)
+		int Enabled = Config.Enabled, Secs = Config.TimeoutSec;
+
+		Error |= setsockopt(Socket.native_handle(), SOL_SOCKET, SO_KEEPALIVE, &Enabled, sizeof(Enabled));
+		Error |= setsockopt(Socket.native_handle(), IPPROTO_TCP, TCP_KEEPALIVE, &Secs, sizeof(Secs));
+
+		if (Result)
+		{
+			ErrorCode = errno;
+		}
+#elif TS_USING(TS_PLATFORM_LINUX)
+		int Enabled = Config.Enabled, KeepCnt = 5, KeepIdle = Config.TimeoutSec, KeepIntvl = Config.IntervalSec;
+
+		Error |= setsockopt(Socket.native_handle(), SOL_SOCKET, SO_KEEPALIVE, &Enabled, sizeof Enabled);
+		Error |= setsockopt(Socket.native_handle(), IPPROTO_TCP, TCP_KEEPCNT, &KeepCnt, sizeof(int));
+		Error |= setsockopt(Socket.native_handle(), IPPROTO_TCP, TCP_KEEPIDLE, &KeepIdle, sizeof(int));
+		Error |= setsockopt(Socket.native_handle(), IPPROTO_TCP, TCP_KEEPINTVL, &KeepIntvl, sizeof(int));
+
+		if (Error)
+		{
+			ErrorCode = errno;
+		}
+#endif
+
+		if (Error)
+		{
+			TS_LOG("Failed to set keep alive for recording socket. Error: %d", ErrorCode);
+		}
+	}
 }
 
 ////////////////////////////////////////////////////////////////////////////////
diff --git a/Engine/Source/Programs/UnrealTraceServer/src/AsioSocket.h b/Engine/Source/Programs/UnrealTraceServer/src/AsioSocket.h
--- a/Engine/Source/Programs/UnrealTraceServer/src/AsioSocket.h
+++ b/Engine/Source/Programs/UnrealTraceServer/src/AsioSocket.h
@@ -11,7 +11,15 @@ class FAsioSocket
 	, public FAsioWriteable
 {
 public:
-							FAsioSocket(asio::ip::tcp::socket& InSocket);
+
+	struct FKeepAliveConfig
+	{
+		bool Enabled = false;
+		uint32 TimeoutSec = 0;
+		uint32 IntervalSec = 0;
+	};
+
+							FAsioSocket(asio::ip::tcp::socket& InSocket, FKeepAliveConfig& KeepAlive = FKeepAliveConfig{});
 	virtual					~FAsioSocket();
 	asio::io_context&		GetIoContext();
 	virtual bool			IsOpen() const override;
diff --git a/Engine/Source/Programs/UnrealTraceServer/src/Recorder.cpp b/Engine/Source/Programs/UnrealTraceServer/src/Recorder.cpp
--- a/Engine/Source/Programs/UnrealTraceServer/src/Recorder.cpp
+++ b/Engine/Source/Programs/UnrealTraceServer/src/Recorder.cpp
@@ -5,10 +5,9 @@
 #include "AsioSocket.h"
 #include "Recorder.h"
 #include "Store.h"
+#include "Logging.h"
+
 
-#if TS_USING(TS_PLATFORM_WINDOWS)
-#	include <Mstcpip.h>
-#endif
 
 ////////////////////////////////////////////////////////////////////////////////
 class FRecorderRelay
@@ -63,48 +62,16 @@ private:
 
 ////////////////////////////////////////////////////////////////////////////////
 FRecorderRelay::FRecorderRelay(asio::ip::tcp::socket& Socket, FStore& InStore)
-: Input(Socket)
-, Store(InStore)
-{
-	Buffer[0] = new uint8[BufferSize];
-	Buffer[1] = new uint8[BufferSize];
-
-#if TS_USING(TS_PLATFORM_WINDOWS)
 	// Trace data is a stream and communication is one way. It is implemented
 	// this way to share code between sending trace data over the wire and writing
 	// it to a file. Because there's no ping/pong we can end up with a half-open
 	// TCP connection if the other end doesn't close its socket. So we'll enable
 	// keep-alive on the socket and set a short timeout (default is 2hrs).
-	tcp_keepalive KeepAlive =
-	{
-		1,		// on
-		15000,	// timeout_ms
-		2000,	// interval_ms
-	};
-
-	DWORD BytesReturned;
-	WSAIoctl(
-		Socket.native_handle(),
-		SIO_KEEPALIVE_VALS,
-		&KeepAlive, sizeof(KeepAlive),
-		nullptr, 0,
-		&BytesReturned,
-		nullptr,
-		nullptr
-	);
-#elif TS_USING(TS_PLATFORM_MAC)
-    int Enabled = 1, Secs = 3;
-    
-    setsockopt(Socket.native_handle(), SOL_SOCKET,  SO_KEEPALIVE, &Enabled, sizeof(Enabled));
-    setsockopt(Socket.native_handle(), IPPROTO_TCP, TCP_KEEPALIVE, &Secs, sizeof(Secs));
-#elif TS_USING(TS_PLATFORM_LINUX)
-    int Enabled = 1, KeepCnt = 5, KeepIdle = 3, KeepIntvl = 1;
-
-    setsockopt(Socket.native_handle(), SOL_SOCKET,  SO_KEEPALIVE, &Enabled, sizeof Enabled);
-    setsockopt(Socket.native_handle(), IPPROTO_TCP, TCP_KEEPCNT, &KeepCnt, sizeof(int));
-    setsockopt(Socket.native_handle(), IPPROTO_TCP, TCP_KEEPIDLE, &KeepIdle, sizeof(int));
-    setsockopt(Socket.native_handle(), IPPROTO_TCP, TCP_KEEPINTVL, &KeepIntvl, sizeof(int));
-#endif
+	: Input(Socket, FAsioSocket::FKeepAliveConfig {true, 15, 2})
+	, Store(InStore)
+{
+	Buffer[0] = new uint8[BufferSize];
+	Buffer[1] = new uint8[BufferSize];
 
 	// Kick things off by reading the magic four bytes at the start of the stream
 	// along with an additional two bytes that are likely the metadata size.
