UnrealTraceServer: Fix crash iterating over a symlink on Windows

Avoids `std::filesystem::directory_entry::is_directory()` throwing an unhandled exception when encountering a symlink.

diff --git a/Engine/Source/Programs/UnrealTraceServer/src/Store.cpp b/Engine/Source/Programs/UnrealTraceServer/src/Store.cpp
--- a/Engine/Source/Programs/UnrealTraceServer/src/Store.cpp
+++ b/Engine/Source/Programs/UnrealTraceServer/src/Store.cpp
@@ -479,7 +479,7 @@ uint32 FStore::FMount::Refresh()
 	std::error_code Ec;
 	for (auto& DirItem : std::filesystem::directory_iterator(Dir, Ec))
 	{
-		if (DirItem.is_directory())
+		if (!DirItem.is_regular_file(Ec) && !DirItem.is_symlink(Ec))
 		{
 			continue;
 		}
